Android SingleInstance Service
Tested on Android with Lyra and also preflight checks.

#android

diff --git a/Engine/Plugins/Runtime/AndroidSingleInstanceService/AndroidSingleInstanceService.uplugin b/Engine/Plugins/Runtime/AndroidSingleInstanceService/AndroidSingleInstanceService.uplugin
new file mode 100644
--- /dev/null
+++ b/Engine/Plugins/Runtime/AndroidSingleInstanceService/AndroidSingleInstanceService.uplugin
@@ -0,0 +1,32 @@
+{
+	"FileVersion": 3,
+	"Version": 1,
+	"VersionName": "1.0",
+	"FriendlyName": "AndroidSingleInstanceService",
+	"Description": "Adds support for Unreal Single Instance Service to Android projects.",
+	"Category": "Android",
+	"CreatedBy": "Epic Games, Inc.",
+	"CreatedByURL": "https://epicgames.com",
+	"DocsURL": "",
+	"MarketplaceURL": "",
+	"SupportURL": "",
+	"CanContainContent": false,
+	"IsBetaVersion": false,
+	"EnabledByDefault": false,
+	"Installed": true,
+	"IsExperimentalVersion": true,
+	"Modules": [
+		{
+			"Name": "AndroidSingleInstanceService",
+			"Type": "Runtime",
+			"LoadingPhase": "PreLoadingScreen",
+			"PlatformAllowList": [ "Android", "Win64", "Mac", "Linux" ]
+		},
+		{
+			"Name": "AndroidSingleInstanceServiceEditor",
+			"Type": "Editor",
+			"LoadingPhase": "PostEngineInit",
+			"PlatformAllowList": [ "Win64", "Mac", "Linux" ]
+		}
+	]
+}
diff --git a/Engine/Plugins/Runtime/AndroidSingleInstanceService/Source/AndroidSingleInstanceService/Android/ASISStub/proguard-android.txt b/Engine/Plugins/Runtime/AndroidSingleInstanceService/Source/AndroidSingleInstanceService/Android/ASISStub/proguard-android.txt
new file mode 100644
--- /dev/null
+++ b/Engine/Plugins/Runtime/AndroidSingleInstanceService/Source/AndroidSingleInstanceService/Android/ASISStub/proguard-android.txt
@@ -0,0 +1,2 @@
+## epicgames wrapper
+-keep class com.epicgames.** { *; }
\ No newline at end of file
diff --git a/Engine/Plugins/Runtime/AndroidSingleInstanceService/Source/AndroidSingleInstanceService/Android/ASISStub/src/main/AndroidManifest.xml b/Engine/Plugins/Runtime/AndroidSingleInstanceService/Source/AndroidSingleInstanceService/Android/ASISStub/src/main/AndroidManifest.xml
new file mode 100644
--- /dev/null
+++ b/Engine/Plugins/Runtime/AndroidSingleInstanceService/Source/AndroidSingleInstanceService/Android/ASISStub/src/main/AndroidManifest.xml
@@ -0,0 +1,36 @@
+<manifest xmlns:android="http://schemas.android.com/apk/res/android">
+	<uses-permission android:name="android.permission.INTERNET" />
+	<uses-permission android:name="android.permission.WRITE_EXTERNAL_STORAGE" />
+	<uses-permission android:name="android.permission.ACCESS_NETWORK_STATE" />
+	<uses-permission android:name="android.permission.WAKE_LOCK" />
+	<uses-permission android:name="com.android.vending.CHECK_LICENSE" />
+	<uses-permission android:name="android.permission.ACCESS_WIFI_STATE" />
+	<uses-permission android:name="android.permission.MODIFY_AUDIO_SETTINGS" />
+	<uses-permission android:name="android.permission.VIBRATE" />
+	<uses-permission android:name="com.android.vending.BILLING" />
+
+	<queries>
+		<package android:name="com.epicgames.AndroidPackagingTest" />
+	</queries>
+	
+	<application
+		android:label="@string/app_name"
+		android:requestLegacyExternalStorage="true"
+		android:hasCode="true" android:usesCleartextTraffic="true" android:allowTaskReparenting="true" android:supportsRtl="true" android:hardwareAccelerated="true"
+		>
+		<activity
+            android:name=".UseServiceActivity"
+			android:hardwareAccelerated="true" android:allowTaskReparenting="true" android:exported="true"
+            android:theme="@style/Theme.UnrealEngine.NoActionBar"android:screenOrientation="landscape">
+			<intent-filter>
+				<action android:name="android.intent.action.MAIN" />
+
+				<category android:name="android.intent.category.LAUNCHER" />
+			</intent-filter>
+		</activity>
+
+	</application>
+	
+ 
+
+</manifest>
diff --git a/Engine/Plugins/Runtime/AndroidSingleInstanceService/Source/AndroidSingleInstanceService/Android/ASISStub/src/main/java/com/epicgames/ASISStub/UseServiceActivity.java b/Engine/Plugins/Runtime/AndroidSingleInstanceService/Source/AndroidSingleInstanceService/Android/ASISStub/src/main/java/com/epicgames/ASISStub/UseServiceActivity.java
new file mode 100644
--- /dev/null
+++ b/Engine/Plugins/Runtime/AndroidSingleInstanceService/Source/AndroidSingleInstanceService/Android/ASISStub/src/main/java/com/epicgames/ASISStub/UseServiceActivity.java
@@ -0,0 +1,758 @@
+package com.epicgames.ASISStub;
+
+import androidx.annotation.NonNull;
+import androidx.annotation.UiThread;
+import androidx.appcompat.app.AppCompatDelegate;
+import androidx.lifecycle.MutableLiveData;
+
+import android.app.Activity;
+import android.app.Application;
+import android.content.ComponentName;
+import android.content.Context;
+import android.content.ContextWrapper;
+import android.content.ServiceConnection;
+import android.graphics.SurfaceTexture;
+import android.os.Build;
+import android.os.Bundle;
+import android.os.Debug;
+import android.os.Handler;
+import android.os.IBinder;
+import android.os.Message;
+import android.os.Messenger;
+import android.os.RemoteException;
+import android.util.Log;
+import android.view.MotionEvent;
+import android.view.Surface;
+import android.view.TextureView;
+import android.view.View;
+import android.content.Intent;
+import android.widget.Button;
+
+import com.epicgames.makeaar.UnrealMessageType;
+
+import java.util.concurrent.atomic.AtomicReference;
+
+
+public class UseServiceActivity extends Activity implements View.OnClickListener
+{
+	private static final Handler.Callback mConnectionStateCallback = (Message msg) -> {
+		final String TAG = "UE-ServiceCallback";
+
+		switch (msg.what) {
+
+			case 0:	//EVENTTYPE_INIT == 0
+				Log.v(TAG, "received mConnectionStateCallback for EVENTTYPE_INIT, obj = " + msg.obj);
+
+				break;
+
+			case 1:	//EVENTTYPE_POST_ENGINE_INIT == 1
+				Log.v(TAG, "received mConnectionStateCallback for EVENTTYPE_POST_ENGINE_INIT, obj = " + msg.obj);
+				break;
+
+			default:
+				//Log.e(TAG, "mConnectionStateCallback, Unknown message = " + msg.what);
+		}
+		return true;
+	};
+
+	//final MutableLiveData<MyServiceConnection> mServiceConnection = new MutableLiveData<MyServiceConnection>(new MyServiceConnection());
+
+	static MyServiceConnection mServiceConnection = null; //new MyServiceConnection();
+
+	static class MyServiceConnection implements ServiceConnection, TextureView.SurfaceTextureListener, TextureView.OnTouchListener
+	{
+		Messenger serviceMessenger = null;
+		Messenger serviceReplyMessenger = null;
+		static MutableLiveData<Boolean> isBoundToService = new MutableLiveData<Boolean>(false);
+
+
+		AtomicReference<Surface> externalSurfaceData = new AtomicReference<Surface>(null);
+
+		int[] surfacePosition = new int[2];
+
+		TextureView prevTextureView = null;
+
+		final int mTaskID;
+
+		int getTaskId() { return mTaskID; }
+
+		ContextWrapper mContextWrapper;
+
+		MyServiceConnection(Activity _activity)
+		{
+//			if (mServiceConnection.getValue() != null)
+//			{
+//				unbindToUnrealInstanceService("MyServiceConnection constructor!!!");
+//				doCleanupForUnbinding("MyServiceConnection constructor!!!");
+//
+//			}
+			assert (_activity != null);
+
+			mContextWrapper = _activity;
+
+			mTaskID = _activity.getTaskId();
+
+//			bindToUnrealInstanceService("MyServiceConnection constructor " + mTaskID);
+		}
+
+
+//		Activity getActivity() {
+//			assert (mContextWrapper instanceof Activity);
+//			return (Activity) mContextWrapper;
+//		}
+//
+//		Application getApplication()
+//		{
+//			Application application = null;
+//			try {
+//				application = getActivity().getApplication();
+//			}
+//			catch (Exception e)
+//			{
+//				e.printStackTrace();
+//			}
+//			finally {
+//				return application;
+//			}
+//		}
+
+		void SendMessage(Message msg) throws RemoteException {
+			if (serviceMessenger != null) {
+				serviceMessenger.send(msg);
+			}
+		}
+
+		void SetTextureView(TextureView textureView_)
+		{
+			textureView_.setAlpha(1.0f);
+			textureView_.setOpaque(false);
+			//textureView_.setOnTouchListener(this);
+			textureView_.setSurfaceTextureListener(this);
+
+			if (textureView_.isAvailable() && textureView_.isAttachedToWindow())
+			{
+				textureView_.getLocationOnScreen(surfacePosition);
+
+				onSurfaceTextureAvailable(textureView_.getSurfaceTexture(), textureView_.getWidth(), textureView_.getHeight());
+			}
+			prevTextureView = textureView_;
+		}
+
+		public boolean onTouch(View v, MotionEvent event) {
+			// TODO Auto-generated method stub
+
+			Log.d(TAG, "Touch: " + event.getAction() + " X: " + event.getX() + " Y: " + event.getY());
+
+			Message msg = Message.obtain(null, UnrealMessageType.TouchEvent.ordinal());
+			msg.getData().putInt("taskId", getTaskId());
+			msg.getData().putParcelable("touch", event);
+
+			try {
+				SendMessage(msg);
+			} catch (RemoteException e) {
+				e.printStackTrace();
+			}
+
+			return true;
+		}
+
+		public void ConsoleCommand(Context context, String consoleCommand)
+		{
+			Log.i(TAG, "ConsoleCommand() - consoleCommand = " + consoleCommand);
+
+			// Create and send a message to the service, using a supported 'what' value
+			Message msg = Message.obtain(null, UnrealMessageType.SendConsoleCommand.ordinal());
+			msg.getData().putString("consoleCommand", consoleCommand);
+
+			try {
+				SendMessage(msg);
+			} catch (RemoteException e) {
+				e.printStackTrace();
+			}
+		}
+
+		void SendData(int event, String param1, int param2, int param3, float param4)
+		{
+			// Create and send a message to the service, using a supported 'what' value
+			Message msg = Message.obtain(null, UnrealMessageType.SendData.ordinal());
+			msg.getData().putInt("taskId", getTaskId());
+			msg.getData().putInt("event", event);
+			msg.getData().putString("param1", param1);
+			msg.getData().putInt("param2", param2);
+			msg.getData().putInt("param3", param3);
+			msg.getData().putFloat("param4", param4);
+
+			try {
+				SendMessage(msg);
+			} catch (RemoteException e) {
+				e.printStackTrace();
+			}
+		}
+
+		void bindToUnrealInstanceService(String caller)
+		{
+			logContextDetails("bindToUnrealInstanceService("+ caller + ") -> ");
+			if(serviceReplyMessenger == null)
+			{
+				serviceReplyMessenger = new Messenger(new Handler(mConnectionStateCallback));
+			}
+
+			Intent intent = new Intent();
+			intent.setClassName(servicePackageName, "com.epicgames.makeaar.UnrealSharedInstanceService");
+
+			intent.putExtra("obbModuleName", obbModuleName);
+			intent.putExtra("obbFileLocation", obbFileLocation);
+			intent.putExtra("commandLineArgs", commandLineArgs);
+			intent.putExtra("enablePropagateAlpha", enablePropagateAlpha);
+
+			if (serviceReplyMessenger != null) {
+				intent.putExtra("callback", serviceReplyMessenger);
+			}
+
+			Log.d(TAG, "bindToUnrealInstanceService intent=" + intent);
+			mContextWrapper.bindService(intent, this,  Context.BIND_AUTO_CREATE | Context.BIND_IMPORTANT | Context.BIND_DEBUG_UNBIND | Context.BIND_ADJUST_WITH_ACTIVITY);
+//			getApplication().bindService(intent, this, Context.BIND_AUTO_CREATE | Context.BIND_IMPORTANT | Context.BIND_ADJUST_WITH_ACTIVITY);
+		}
+
+		void unbindToUnrealInstanceService(String caller)
+		{
+			logContextDetails("unbindToUnrealInstanceService("+ caller + ") -> ");
+
+			if (isBoundToService.getValue()) {
+				Log.v(TAG, "unbindToUnrealInstanceService");
+
+				mContextWrapper.unbindService(this);
+			}
+		}
+
+
+		private void logContextDetails(String caller)
+		{
+			Log.w(TAG, caller + ": isBoundToService=" +isBoundToService.getValue()
+					+ ", taskID=" +getTaskId()
+					+ ", serviceMessenger=" +serviceMessenger
+					+ ", serviceReplyMessenger=" +serviceReplyMessenger
+					+ ", externalSurfaceData=" +externalSurfaceData
+			);
+		}
+
+		private void doCleanupForUnbinding(String caller) {
+			logContextDetails("doCleanupForBinding(" + caller + ")");
+
+//			if (isBoundToService.getValue())
+//			{
+//				//getApplication().unbindService(this);
+//				//mServiceConnection.setValue( null );
+//			}
+
+			if (prevTextureView != null) {
+				prevTextureView.setSurfaceTextureListener(null);
+				prevTextureView = null;
+			}
+
+			if (externalSurfaceData.get() != null)
+			{
+				externalSurfaceData.get().release();
+				externalSurfaceData.set(null);
+			}
+
+			serviceMessenger = null;
+			serviceReplyMessenger = null;
+			isBoundToService.setValue( false );
+		}
+
+		@Override
+		public void onServiceConnected(ComponentName name, IBinder serviceBinder) {
+			Log.v(TAG, "onServiceConnected name=" + name + ", serviceBinder=" + serviceBinder);
+			serviceMessenger = new Messenger(serviceBinder);
+			//serviceReplyMessenger = new Messenger(new Handler(mConnectionStateCallback));
+			isBoundToService.setValue( true );
+//			TextureView textureView_ = (TextureView) findViewById( R.id.renderView );
+//
+//			textureView_ = (TextureView) getActivity().findViewById( R.id.renderView );
+//			SetTextureView(textureView_);
+		}
+
+		@Override
+		public void onServiceDisconnected(ComponentName name) {
+
+			Log.w(TAG, "onServiceDisconnected name=" + name);
+
+			doCleanupForUnbinding("onServiceDisconnected name=" + name);
+		}
+
+		@Override
+		public void onBindingDied(ComponentName name) {
+
+			Log.w(TAG, "onBindingDied name=" + name);
+
+			doCleanupForUnbinding("onBindingDied name=" + name);
+
+			ServiceConnection.super.onBindingDied(name);
+		}
+
+		@Override
+		public void onNullBinding(ComponentName name) {
+
+			Log.w(TAG, "onNullBinding name=" + name);
+
+			//doCleanupForBinding("onNullBinding");
+
+			ServiceConnection.super.onNullBinding(name);
+		}
+
+		@Override
+		public void onSurfaceTextureAvailable(@NonNull SurfaceTexture surfaceTexture, int width, int height) {
+
+			logContextDetails("onSurfaceTextureAvailable(" + surfaceTexture + ")");
+
+			int[] surfacePosition = this.surfacePosition;
+			int[] surfaceSize = new int[]{width, height};
+
+			if (externalSurfaceData.get() != null) {
+				externalSurfaceData.get().release();
+			}
+
+			externalSurfaceData.set(new Surface(surfaceTexture));
+
+			attachSurfaceToService(getTaskId(), externalSurfaceData.get(), surfacePosition, surfaceSize);
+		}
+
+		@Override
+		public void onSurfaceTextureSizeChanged(@NonNull SurfaceTexture surfaceTexture, int width, int height) {
+			logContextDetails("onSurfaceTextureSizeChanged(" + surfaceTexture + ")");
+		}
+
+		// Invoked when the specified SurfaceTexture is about to be destroyed.
+		// If returns true, no rendering should happen inside the surface texture after this method is invoked.
+		// If returns false, the client needs to call SurfaceTexture.release().
+		// Most applications should return true.
+		@Override
+		public boolean onSurfaceTextureDestroyed(@NonNull SurfaceTexture surfaceTexture) {
+
+			logContextDetails("onSurfaceTextureDestroyed(" + surfaceTexture + ")");
+			detachSurfaceFromService(getTaskId(), null);
+
+			if (externalSurfaceData.get() != null) {
+				externalSurfaceData.get().release();
+				externalSurfaceData.set(null);
+			}
+
+			if (prevTextureView != null) {
+				prevTextureView.setSurfaceTextureListener(null);
+			}
+			return true;
+		}
+
+		@Override
+		public void onSurfaceTextureUpdated(@NonNull SurfaceTexture surfaceTexture) {
+		}
+
+		@UiThread
+		void attachSurfaceToService(int attachId, Surface externalSurface, int[] viewPosition, int[] viewSize)
+		{
+			if (externalSurface ==  null) {
+				Log.w(TAG, "attachSurfaceToService, externalSurface is null");
+				return;
+			}
+
+			if (Build.VERSION.SDK_INT >= Build.VERSION_CODES.P) {
+				Log.v(TAG, "attachSurfaceToService() - " + "proc = " + Application.getProcessName() + ", externalSurface = " + externalSurface + ", attachId = " + attachId);
+			}
+
+			Message msg = Message.obtain(null, UnrealMessageType.AttachExternalSurface.ordinal());
+			Bundle data = msg.getData();
+			data.putInt("attachId", attachId);
+			data.putBoolean("enablePropagateAlpha", true);
+
+			data.putParcelable("surface", externalSurface);
+
+			if (viewPosition.length == 2) {
+				data.putIntArray("viewPosition", viewPosition);
+			}
+
+			if (viewSize.length == 2) {
+				data.putIntArray("viewSize", viewSize);
+
+			}
+
+			try {
+				SendMessage(msg);
+			} catch (RemoteException e) {
+				e.printStackTrace();
+			}
+		}
+
+		@UiThread
+		void detachSurfaceFromService(int attachId, Handler.Callback callback) {
+
+			Message msg = Message.obtain(
+					null,
+					UnrealMessageType.DetachExternalSurface.ordinal()
+			);
+
+			msg.getData().putInt("attachId", attachId);
+
+			if (callback != null) {
+
+				Handler handler = new Handler(callback);
+				Messenger replyMessenger = new Messenger(handler);
+
+				msg.replyTo = replyMessenger;
+			}
+
+			Log.i(TAG, "detachSurfaceFromService() - " + "proc = ${Application.getProcessName()}");
+
+			try {
+				SendMessage(msg);
+			} catch (RemoteException e) {
+				e.printStackTrace();
+			}
+		}
+
+		@UiThread
+		void resumeService(int attachId)
+		{
+			if (isBoundToService.getValue()) {
+				Message msg = Message.obtain(null, UnrealMessageType.ResumeService.ordinal());
+				Bundle data = msg.getData();
+				data.putInt("attachId", attachId);
+
+				try {
+					SendMessage(msg);
+				} catch (RemoteException e) {
+					e.printStackTrace();
+				}
+			}
+		}
+	}
+
+	MyServiceConnection GetServiceConnection(String caller)
+	{
+
+		Log.w(TAG, "GetServiceConnection(" + caller + ") : mServiceConnection = " + mServiceConnection);
+
+		assert (mServiceConnection != null);
+
+		return mServiceConnection;
+	}
+
+	MyServiceConnection CreateServiceConnection(Activity activity, String caller)
+	{
+
+		Log.w(TAG, "CreateServiceConnection(" + caller + ") : mServiceConnection = " + mServiceConnection + ", activity=" + activity);
+
+		if (mServiceConnection != null)
+		{
+			mServiceConnection.logContextDetails("CreateServiceConnection destroy prev (" + activity + ")");
+			DeposeServiceConnection();
+		}
+
+		if (mServiceConnection == null)
+		{
+			//mServiceConnection.logContextDetails("CreateServiceConnection create new (" + activity + ")");
+			mServiceConnection = new MyServiceConnection(activity);
+			mServiceConnection.bindToUnrealInstanceService("CreateServiceConnection");
+		}
+
+		return mServiceConnection;
+	}
+
+	void DeposeServiceConnection()
+	{
+
+		Log.w(TAG, "DeposeServiceConnection, mServiceConnection = " + mServiceConnection);
+
+		if (mServiceConnection != null) {
+			mServiceConnection.unbindToUnrealInstanceService("DeposeServiceConnection");
+			//mServiceConnection.doCleanupForUnbinding("DeposeServiceConnection");
+		}
+
+		mServiceConnection = null;
+	}
+
+
+
+	static Integer next_theme = AppCompatDelegate.MODE_NIGHT_NO;
+
+	// declaring objects of Button class
+	// Surface renderSurface;
+
+	final static String TAG = "UE-UseServiceActivity";
+
+	final static String obbModuleName = "AndroidPackagingTest";//""ALPHA_POC";
+	final static String servicePackageName = "com.epicgames.AndroidPackagingTest";
+	final static String obbFileLocation = ""; //"/data/local/tmp/3d/AndroidPackagingTest.main.obb.png";
+
+	final static String commandLineArgs = "-nosound -launchandroidflags=0 -dpcvars='Android.UseGameThread=1'";
+	final static boolean enablePropagateAlpha = true;
+
+
+
+	private Button activateView1Button, activateView2Button, activateView3Button, restartActivityButton, themeChangeButton, resumeButton;
+	private Button pauseButton, stopButton, bindServiceButton, unbindServiceButton, memReportButton;
+
+	void ActivateTextureView(int viewIndex)
+	{
+		TextureView textureView_ = null;
+		switch (viewIndex)
+		{
+			case 1:
+				textureView_ = (TextureView) findViewById( R.id.textureView );
+				break;
+			case 2:
+				textureView_ = (TextureView) findViewById( R.id.renderView );
+				break;
+			case 3:
+				textureView_ = (TextureView) findViewById( R.id.renderView2 );
+				break;
+		}
+
+		GetServiceConnection("ActivateTextureView").SetTextureView(textureView_);
+	}
+
+
+	@UiThread
+	void stopService(int attachId, Handler.Callback callback)
+	{
+		Intent intent = new Intent();
+		intent.setClassName(servicePackageName, "com.epicgames.makeaar.UnrealSharedInstanceService");
+		getApplication().stopService(intent);
+	}
+
+	@Override
+	public void onContentChanged() {
+		Log.v(TAG, "UseServiceActivity: onContentChanged");
+		super.onContentChanged();
+	}
+
+	@Override
+	protected void onResume() {
+		Log.v(TAG, "UseServiceActivity: onResume");
+		super.onResume();
+
+//		GetServiceConnection("onResume").resumeService(getTaskId());
+	}
+
+	@Override
+	protected void onPause() {
+		Log.v(TAG, "UseServiceActivity: onPause");
+		super.onPause();
+	}
+
+	@Override
+	protected void onStart() {
+		Log.v(TAG, "UseServiceActivity: onStart");
+		super.onStart();
+
+//		CreateServiceConnection( this,"UseServiceActivity::onStart()");
+	}
+
+	@Override
+	protected void onStop() {
+		Log.v(TAG, "UseServiceActivity: onStop");
+		super.onStop();
+//		DeposeServiceConnection();
+	}
+
+	@Override
+	protected void onDestroy() {
+		Log.v(TAG, "UseServiceActivity: onDestroy");
+		super.onDestroy();
+		//DeposeServiceConnection();
+		//GetServiceConnection().doCleanupForBinding("onDestroy");
+	}
+
+	@Override
+	protected void onRestart() {
+		Log.v(TAG, "UseServiceActivity onRestart called! hasBeenCreated = " + hasBeenCreated);
+		super.onRestart();
+	}
+
+	@Override
+	protected void onCreate(Bundle savedInstanceState) {
+		super.onCreate( savedInstanceState );
+		AppCompatDelegate.setDefaultNightMode(next_theme);
+		getTheme().applyStyle(R.style.Theme_UnrealEngine_NoActionBar, true);
+
+		setContentView( R.layout.interface_use_service_activity );
+
+
+
+		activateView1Button 	= (Button) findViewById( R.id.activateView1Button);
+		activateView2Button 	= (Button) findViewById( R.id.activateView2Button);
+		activateView3Button 	= (Button) findViewById( R.id.activateView3Button);
+
+		pauseButton 			= (Button) findViewById( R.id.buttonPause);
+		stopButton 				= (Button) findViewById( R.id.buttonStopService);
+		bindServiceButton 		= (Button) findViewById( R.id.buttonBindService);
+		unbindServiceButton 	= (Button) findViewById( R.id.buttonUnbindService);
+		restartActivityButton 	= (Button) findViewById( R.id.restartActivityButton);
+		memReportButton 		= (Button) findViewById( R.id.buttonMemReport);
+		themeChangeButton 		= (Button) findViewById( R.id.changeThemeButton);
+		resumeButton 			= (Button) findViewById( R.id.resumeButton);
+
+
+		activateView1Button.setOnClickListener( this );
+		activateView2Button.setOnClickListener( this );
+		activateView3Button.setOnClickListener( this );
+
+		pauseButton.setOnClickListener( this );
+		stopButton.setOnClickListener( this );
+		bindServiceButton.setOnClickListener( this );
+		unbindServiceButton.setOnClickListener( this );
+		restartActivityButton.setOnClickListener( this );
+		memReportButton.setOnClickListener( this );
+		themeChangeButton.setOnClickListener( this );
+		resumeButton.setOnClickListener( this );
+
+
+		Log.v(TAG, "UseServiceActivity onCreate called! hasBeenCreated = " + hasBeenCreated);
+
+		if (false == hasBeenCreated) {
+			Debug.getMemoryInfo(baseMemoryInfo);
+			Debug.getMemoryInfo(prevMemoryInfo);
+			hasBeenCreated = true;
+		}
+
+		CreateServiceConnection( this,"UseServiceActivity::onCreate()");
+	}
+
+	public void onClick(View view)
+	{
+		final boolean waitForBind = false;
+
+		if (activateView1Button.equals(view)) {
+			ActivateTextureView(1);
+		} else if (activateView2Button.equals(view)) {
+			ActivateTextureView(2);
+		} else if (activateView3Button.equals(view)) {
+			ActivateTextureView(3);
+		} else if (pauseButton.equals(view)) {
+
+			try {
+				GetServiceConnection("Pause Button").detachSurfaceFromService(getTaskId(), null);
+			} catch (Exception e) {
+				e.printStackTrace();
+			}
+
+		} else if (stopButton.equals(view)) {
+			stopService(getTaskId(), null);
+		} else if (bindServiceButton.equals(view)) {
+
+//			try {
+//				GetServiceConnection("Bind Service Button").bindToUnrealInstanceService("onClick");
+//			} catch (Exception e) {
+//				e.printStackTrace();
+//			}
+
+		} else if (unbindServiceButton.equals(view)) {
+
+			try {
+				DeposeServiceConnection();
+				//GetServiceConnection("Unbind Service Button").unbindToUnrealInstanceService("onClick");
+			} catch (Exception e) {
+				e.printStackTrace();
+			}
+
+		} else if (restartActivityButton.equals(view)) {
+			recreate();
+		} else if (memReportButton.equals(view)) {
+			ReportDebugMemoryInfo();
+		} else if (themeChangeButton.equals(view)) {
+			handleThemeChange();
+		} else if (resumeButton.equals(view)) {
+
+			try {
+				GetServiceConnection("Resume Button").resumeService(getTaskId());
+			} catch (Exception e) {
+				e.printStackTrace();
+			}
+		}
+
+	}
+
+
+
+//    @Override
+//  public boolean onTouchEvent(MotionEvent e) {
+//      return ueInterface.MessageServiceToHandleTouch(e);
+//    }
+
+
+	//private final UEInterfaceHelpers ueInterface = new UEInterfaceHelpers(this);
+
+	int cycleThemeCounter = 0;
+	static boolean hasBeenCreated = false;
+	static Debug.MemoryInfo memoryInfo = new Debug.MemoryInfo();
+	static Debug.MemoryInfo prevMemoryInfo = new Debug.MemoryInfo();
+	static Debug.MemoryInfo baseMemoryInfo = new Debug.MemoryInfo();
+	static int tryCount = 0;
+	void ReportDebugMemoryInfo()
+	{
+//        if (prevMemoryInfo != null)
+//        {
+//            Log.i(TAG, "Debug.memoryInfo prev= " + prevMemoryInfo.getMemoryStats());
+//
+//    }
+		Runtime.getRuntime().gc();
+
+		Debug.getMemoryInfo(memoryInfo);
+		Log.v(TAG, "Debug.memoryInfo:\n\tthis= " + memoryInfo.getMemoryStats()
+				+ "\n\tprev= " + prevMemoryInfo.getMemoryStats()
+				+ "\n\tbase= " + baseMemoryInfo.getMemoryStats()
+		);
+		Debug.getMemoryInfo(prevMemoryInfo);
+
+	}
+
+	void doThemeChange(Integer value)
+	{
+		Log.v(TAG, "doing a touch event theme change: new uiTheme = ${value}");
+//        prev_uiMode = next_theme;
+		next_theme = value;
+//        AppCompatDelegate.setDefaultNightMode(value);
+		recreate();
+
+	}
+
+	void handleThemeChange()
+	{
+		//GetServiceConnection().bindToUnrealInstanceService("handleThemeChange");
+		//resumeService(getTaskId(), null);
+
+		doThemeChange(next_theme == AppCompatDelegate.MODE_NIGHT_NO ? AppCompatDelegate.MODE_NIGHT_YES : AppCompatDelegate.MODE_NIGHT_NO);
+
+
+		tryCount = 110;
+//            ueInterface.bindToUnrealInstanceService(this, textureView_, BuildConfig.OBB_FILE_NAME, BuildConfig.OBB_MODULE_NAME, "-nosound -nocrashreports -opengl -featureleveles31", true, mConnectionStateCallback);
+//            ReportDebugMemoryInfo();
+
+
+//            ueInterface.runOnMainThread(new Runnable()
+		runOnUiThread(new Runnable()
+		{
+			@Override
+			public void run() {
+				Context context = getApplicationContext();
+				synchronized (context) {
+
+					try {
+						while ( (++tryCount) < 15) {
+							Log.v(TAG, "mem pressure test iter = " + tryCount);
+
+
+							doThemeChange(next_theme == AppCompatDelegate.MODE_NIGHT_NO ? AppCompatDelegate.MODE_NIGHT_YES : AppCompatDelegate.MODE_NIGHT_NO);
+
+						}
+					} catch (Exception e) {
+						e.printStackTrace();
+					}
+
+				}
+			}
+		});
+	}
+
+
+}
diff --git a/Engine/Plugins/Runtime/AndroidSingleInstanceService/Source/AndroidSingleInstanceService/Android/ASISStub/src/main/res/values/strings.xml b/Engine/Plugins/Runtime/AndroidSingleInstanceService/Source/AndroidSingleInstanceService/Android/ASISStub/src/main/res/values/strings.xml
new file mode 100644
--- /dev/null
+++ b/Engine/Plugins/Runtime/AndroidSingleInstanceService/Source/AndroidSingleInstanceService/Android/ASISStub/src/main/res/values/strings.xml
@@ -0,0 +1,3 @@
+<resources>
+    <string name="app_name">ASISStub</string>
+</resources>
diff --git a/Engine/Plugins/Runtime/AndroidSingleInstanceService/Source/AndroidSingleInstanceService/Android/buildsystem/UnrealAARVersion.xml b/Engine/Plugins/Runtime/AndroidSingleInstanceService/Source/AndroidSingleInstanceService/Android/buildsystem/UnrealAARVersion.xml
new file mode 100644
--- /dev/null
+++ b/Engine/Plugins/Runtime/AndroidSingleInstanceService/Source/AndroidSingleInstanceService/Android/buildsystem/UnrealAARVersion.xml
@@ -0,0 +1,10 @@
+<?xml version="1.0" encoding="utf-8"?>
+<resources>
+	<!-- ClientId.MajorVersion.MinorVersion -->
+	<!-- ClientId = 05.2.31 for 05 = UE5 main, 02 version 5.2 when code was first integrated, 31 is the revision number from QNX HMI //UE4/Private-QNX-4.25/ that was used as initial dev point for UE5 integration.
+		 MajorVersion = 2 for top level UE5 main
+		 MinorVersion = 31, merge in initial code base and update for UE5 from 01.7.31 and single instance work	
+		 MinorVersion = 34, merge in code changes made in 4.25 qnx-brave (54.8.34) that address early lifecycle issues and focus changes.
+	-->
+	<string name="UnrealEngineAARVersion">05.2.34</string>
+</resources>
