Android SingleInstance Service

diff --git a/Engine/Plugins/Runtime/AndroidSingleInstanceService/Source/AndroidSingleInstanceService/Android/src/main/java/com/epicgames/makeaar/Engine.java b/Engine/Plugins/Runtime/AndroidSingleInstanceService/Source/AndroidSingleInstanceService/Android/src/main/java/com/epicgames/makeaar/Engine.java
new file mode 100644
--- /dev/null
+++ b/Engine/Plugins/Runtime/AndroidSingleInstanceService/Source/AndroidSingleInstanceService/Android/src/main/java/com/epicgames/makeaar/Engine.java
@@ -0,0 +1,1970 @@
+package com.epicgames.makeaar;
+
+import static android.content.Context.CONTEXT_IGNORE_SECURITY;
+import static android.content.Context.CONTEXT_INCLUDE_CODE;
+
+import androidx.annotation.NonNull;
+import androidx.annotation.UiThread;
+import android.annotation.SuppressLint;
+import android.app.Activity;
+import android.app.ActivityManager;
+import android.content.Context;
+import android.content.ContextWrapper;
+import android.content.Intent;
+import android.content.pm.PackageManager;
+import android.content.res.Configuration;
+import android.graphics.Bitmap;
+import android.graphics.PixelFormat;
+import android.graphics.Point;
+import android.graphics.Rect;
+import android.graphics.drawable.Drawable;
+import android.os.Build;
+import android.util.Log;
+import android.view.MotionEvent;
+import android.view.Surface;
+import android.view.SurfaceHolder;
+import android.view.SurfaceView;
+import android.view.TextureView;
+import android.view.View;
+import android.view.ViewGroup;
+import android.graphics.SurfaceTexture;
+import android.view.ViewManager;
+import android.graphics.Canvas;
+import android.view.ViewTreeObserver;
+import com.epicgames.unreal.GameActivity;
+import com.epicgames.unreal.SimpleContextWrapper;
+
+import java.lang.ref.WeakReference;
+import java.util.Arrays;
+import java.util.Objects;
+import java.util.Queue;
+
+import static java.lang.System.loadLibrary;
+
+public class Engine
+{
+	public static int EVENTTYPE_INIT = 0;
+	public static int EVENTTYPE_POST_ENGINE_INIT = 1;
+	public static int EVENTTYPE_ENGINELOOP_INIT_COMPLETE = 2;
+	public static int EVENTTYPE_FRAME_BEGIN = 3;
+	public static int EVENTTYPE_FRAME_END = 4;
+	public static int EVENTTYPE_PRE_LOAD_MAP = 5;
+	public static int EVENTTYPE_POST_LOAD_MAP = 6;
+	public static int EVENTTYPE_ACTION = 7;
+	public static int EVENTTYPE_ACTIVITY_HAS_CHANGED = 8;
+
+	private static final String TAG = "UE_ASIS_Engine";
+
+
+	static boolean CheckEvent(long event)
+	{
+		return (maskEngineEvents & (1<<event)) != 0;
+	}
+
+	public interface IEventCallback
+	{
+		IEventCallback NO_OP = new IEventCallback() {
+			@Override
+			public void eventCallback(int event, String param1, int param2, int param3, float param4) {
+				// no op
+			}
+		};
+		void eventCallback(int event, String param1, int param2, int param3, float param4);
+	}
+
+	private static WeakReference<IEventCallback> mCallbackRef = new WeakReference<>(IEventCallback.NO_OP);
+
+	interface IViewInterface
+	{
+		public void GetPosSize(int[] pos, int[] size);
+		public void SetupAlpha(Boolean enablePropagateAlpha);
+		void Detach();
+	}
+
+	public static class ViewInterface<T> implements IViewInterface, SurfaceHolder.Callback, TextureView.SurfaceTextureListener, ViewTreeObserver.OnWindowFocusChangeListener//, ViewTreeObserver.OnWindowVisibilityChangeListener
+	{
+		private static final String TAG = "UE_ASIS_ViewInterface";
+		private T view;
+
+		private Class<? extends Object> type;
+
+		public ViewInterface(T inView) {
+			this.type = inView.getClass();
+			this.view = inView;
+
+			((View)inView).getViewTreeObserver().addOnWindowFocusChangeListener(this);
+			//((View)inView).getViewTreeObserver().addOnWindowVisibilityChangeListener(this);
+		}
+
+		public void onWindowVisibilityChanged(int visibility)
+		{
+			if (visibility == View.VISIBLE) {
+				QueueResume("ViewInterface::onWindowVisibilityChanged handle VISIBLE");
+			} else {
+				QueuePause("ViewInterface::onWindowVisibilityChanged handle gone or invisible, visibility=" + visibility);
+			}
+		}
+
+		public void onWindowFocusChanged(boolean hasFocus)
+		{
+			Engine engine = Get();
+			Log.i(TAG, "** ViewInterface::handleViewFocusChange hasFocus=" + hasFocus
+					+ ", maskEngineEvents=" + maskEngineEvents
+					//+ ", bGameActivityNativeMainIsDone=" + bGameActivityNativeMainIsDone
+					+ ", bCalledStart:" + engine.bCalledStart
+					+ ", startCounter:" + startCounter
+					+ ", bResuming:" + bResuming
+					+ ", onPauseCounter:" + onPauseCounter
+					+ ", renderSurfaceValid:" + engine.renderSurfaceValid
+					+ ", renderSurface:" + engine.renderSurface
+					+ ", pendingChangedSurface:" + engine.pendingChangedSurface
+			);
+
+			if (!hasFocus)
+			{
+				if ( !CheckEvent(EVENTTYPE_INIT) ) {
+					QueuePause("ViewInterface::handleViewFocusChange is false and EVENTTYPE_INIT has not been called, hasFocus=" + hasFocus);
+				}
+				else
+				{
+					QueuePause("ViewInterface::handleViewFocusChange, hasFocus=" + hasFocus);
+				}
+			}
+			else
+			{
+				QueueResume("ViewInterface::handleViewFocusChange, hasFocus=" + hasFocus + ", maskEngineEvents=" + maskEngineEvents);
+			}
+		}
+
+		public boolean canAcceptType(Class<?> candidate) {
+			return type.isAssignableFrom(candidate);
+		}
+
+		public boolean operator_not_equals(final Class<?> e2) {
+			return this.view != e2;
+		}
+
+		public void GetPosSize(int[] pos, int[] size)
+		{
+			if (pos.length == 2) {
+				((View) view).getLocationInWindow(pos);
+			}
+
+			if (size.length == 2) {
+				size[0] = ((View) view).getWidth();
+				size[1] = ((View) view).getHeight();
+			}
+		}
+
+		public void SetupAlpha(Boolean enablePropagateAlpha)
+		{
+			if (view instanceof SurfaceView)
+			{
+				SurfaceView surfaceView = ((SurfaceView)view);
+				surfaceView.setZOrderOnTop(true);
+				surfaceView.setAlpha(1.0f);
+				if (enablePropagateAlpha)
+				{
+					surfaceView.getHolder().setFormat(PixelFormat.TRANSPARENT);
+				}
+				surfaceView.getHolder().addCallback(this);
+			}
+			else if (view instanceof TextureView)
+			{
+				TextureView textureView = ((TextureView)view);
+				textureView.setAlpha(1.0f);
+				textureView.setOpaque(!enablePropagateAlpha);
+				textureView.setSurfaceTextureListener(this);
+			}
+		}
+
+
+		@Override
+		public void surfaceCreated(SurfaceHolder holder) {
+			int[] surfacePosition = new int[] {holder.getSurfaceFrame().left, holder.getSurfaceFrame().top};
+			int[] surfaceSize = new int[]{holder.getSurfaceFrame().width(), holder.getSurfaceFrame().height()};
+
+			Log.i(TAG, "surfaceCreated, surfacePosition = " + Arrays.toString(surfacePosition) + ", surfaceSize = " + Arrays.toString(surfaceSize) + ", holder.getSurface().isValid()= " + holder.getSurface().isValid());
+
+			Engine engine = Get();
+
+			synchronized (syncLock) {
+				engine.renderSurfaceViewPosition = surfacePosition;
+				engine.renderSurfaceViewSize = surfaceSize;
+				if (holder.getSurface().isValid() )//&& engine.renderSurfaceValid && engine.renderSurface != null)
+				{
+					engine.releaseVolatileResources("surfaceCreated and holder has valid surface!");
+				}
+				engine.pendingChangedSurface = holder.getSurface();
+			}
+		}
+
+		@Override
+		public void surfaceChanged(SurfaceHolder holder, int format, int width, int height) {
+			Log.i(TAG, "surfaceChanged, format = " + format + ", width = " + width + ", bResuming=" + bResuming + ", onPauseCounter=" + onPauseCounter + ", renderSurfaceValid=" + Get().renderSurfaceValid + ", renderSurface=" + Get().renderSurface );
+
+			int[] surfacePosition = new int[] {holder.getSurfaceFrame().left, holder.getSurfaceFrame().top};
+			int[] surfaceSize = new int[]{width, height};
+
+			if (view != null)
+			{
+				GetPosSize(surfacePosition, surfaceSize);
+				Log.i(TAG, "surfaceChanged with view, format = " + format + ", surfacePosition = " + Arrays.toString(surfacePosition) + ", surfaceSize = " + Arrays.toString(surfaceSize));
+			}
+
+			Engine engine = Get();
+
+			//synchronized(syncLock)
+			{
+				if (holder.getSurface() != engine.renderSurface && engine.renderSurface != null)
+				{
+					Log.w(TAG, "surfaceChanged and releaseVolatileResources cause holder surface is different than prev");
+
+					engine.releaseVolatileResources("surfaceChanged cause holder surface is different than prev");
+				}
+
+				engine.renderSurfaceViewPosition = surfacePosition;
+				engine.renderSurfaceViewSize = surfaceSize;
+				engine.pendingChangedSurface = holder.getSurface();
+				bResuming |= engine.pendingChangedSurface != null;
+			}
+			// queue the resume if needed at the end of surfaceChanged...
+			if (bResuming)
+			{
+				if (!engine.bCalledStart)
+				{
+					QueueStart("surfaceChanged");
+				}
+				QueueResume("surfaceChanged" );
+			}
+		}
+
+		@Override
+		public void surfaceDestroyed(SurfaceHolder holder) {
+
+			Engine engine = Get();
+
+			Log.i(TAG, "surfaceDestroyed, holder = " + holder
+				+ ", lifecycleContextID=" + lifecycleContextID
+				+ ", bResuming=" + bResuming
+				+ ", surfaceCreatedCounter=" + surfaceCreatedCounter
+				+ ", onPauseCounter=" + onPauseCounter
+				+ ", renderSurfaceValid=" + engine.renderSurfaceValid
+				+ ", renderSurface=" + engine.renderSurface
+				+ ", holder.surface=" + holder.getSurface()
+			);
+
+			//synchronized(syncLock)
+			{
+				if (engine.renderSurfaceValid && engine.renderSurface == holder.getSurface()) {
+					Log.i(TAG, "surfaceDestroyed, doing cleanup and pausing, holder = " + holder
+							+ ", lifecycleContextID=" + lifecycleContextID
+							+ ", bResuming=" + bResuming
+							+ ", surfaceCreatedCounter=" + surfaceCreatedCounter
+							+ ", onPauseCounter=" + onPauseCounter
+							+ ", renderSurfaceValid=" + engine.renderSurfaceValid
+							+ ", renderSurface=" + engine.renderSurface
+							+ ", holder.surface=" + holder.getSurface()
+					);
+					Surface pendingSurface = holder.getSurface();
+					GetOverrideGameActivity().nativeSetSurfaceOverride(null, renderSurfaceViewPosition[0], renderSurfaceViewPosition[1]);
+					if (engine.renderSurface != null && (onPauseCounter > 0 || startCounter > 0)) {
+						GetOverrideGameActivity().nativeAppCommand(APP_CMD_TERM_WINDOW);
+					}
+
+					if (engine.renderSurfaceValid && (startCounter > 0 || onPauseCounter > 0)) {
+						bResuming = true;
+					}
+
+					engine.renderSurface = engine.pendingChangedSurface;
+					engine.renderSurfaceValid = engine.renderSurface != null;
+					engine.pendingChangedSurface = pendingSurface;
+
+					engine.onPause(lifecycleContextID, "surfaceDestroyed case 1");
+
+
+				} else if (holder.getSurface() != null && engine.renderSurface != null) {
+					Log.w(TAG, "surfaceDestroyed, re-bind surface, holder = " + holder
+							+ ", lifecycleContextID=" + lifecycleContextID
+							+ ", bResuming=" + bResuming
+							+ ", surfaceCreatedCounter=" + surfaceCreatedCounter
+							+ ", onPauseCounter=" + onPauseCounter
+							+ ", renderSurfaceValid=" + engine.renderSurfaceValid
+							+ ", renderSurface=" + engine.renderSurface
+							+ ", holder.surface=" + holder.getSurface()
+					);
+
+					engine.pendingChangedSurface = holder.getSurface();
+
+					if (engine.pendingChangedSurface != null) {
+						bResuming = true;
+						
+						engine.QueueResume("surfaceDestroyed valid pendingChangedSurface");
+					}
+				}
+			}
+		}
+
+		private Surface surfaceForTextureView = null;
+
+		@Override
+		public void onSurfaceTextureAvailable(SurfaceTexture surfaceTexture, int width, int height) {
+
+			//synchronized(syncLock)
+			{
+				TextureView textureView = ((TextureView) view);
+
+				int[] surfacePosition = new int[]{0, 0};
+				textureView.getLocationOnScreen(surfacePosition);
+				int[] surfaceSize = new int[]{width, height};
+				Log.w(TAG, "onSurfaceTextureAvailable, surfacePosition = " + Arrays.toString(surfacePosition) + ", surfaceSize = " + Arrays.toString(surfaceSize));
+
+				renderSurfaceViewPosition = surfacePosition;
+				renderSurfaceViewSize = surfaceSize;
+
+
+				if (surfaceForTextureView != null) {
+					surfaceForTextureView.release();
+					surfaceForTextureView = null;
+				}
+
+				Engine.Get().pendingChangedSurface = new Surface(surfaceTexture);
+				surfaceForTextureView = Engine.Get().pendingChangedSurface;
+			}
+			Engine.Get().QueueResume("onSurfaceTextureAvailable");
+		}
+
+		@Override
+		public void onSurfaceTextureSizeChanged(SurfaceTexture surfaceTexture, int width, int height) {
+			//synchronized(syncLock)
+			{	
+				TextureView textureView = ((TextureView) view);
+
+				int[] surfacePosition = new int[]{0, 0};
+				textureView.getLocationOnScreen(surfacePosition);
+				int[] surfaceSize = new int[]{width, height};
+				Log.i(TAG, "onSurfaceTextureSizeChanged, surfacePosition = " + Arrays.toString(surfacePosition) + ", surfaceSize = " + Arrays.toString(surfaceSize));
+
+				renderSurfaceViewSize = surfaceSize;
+				GetOverrideGameActivity().nativeSetSurfaceViewInfo(renderSurfaceViewSize[0], renderSurfaceViewSize[1]);
+			}
+		}
+
+		// Invoked when the specified SurfaceTexture is about to be destroyed.
+		// If returns true, no rendering should happen inside the surface texture after this method is invoked.
+		// If returns false, the client needs to call SurfaceTexture.release().
+		// Most applications should return true.
+		@Override
+		public boolean onSurfaceTextureDestroyed(SurfaceTexture surfaceTexture) {
+
+			//synchronized(syncLock)
+			{
+
+				Log.i(TAG, "onSurfaceTextureDestroyed, surfaceForTextureView = " + surfaceForTextureView + ", renderSurface = " + Engine.Get().renderSurface);
+
+				Engine.Get().releaseVolatileResources("onSurfaceTextureDestroyed called");
+				if (surfaceForTextureView != null)
+				{
+					surfaceForTextureView.release();
+					surfaceForTextureView = null;
+				}
+
+			}
+			return true;
+		}
+
+		@Override
+		public void onSurfaceTextureUpdated(SurfaceTexture surfaceTexture) {
+		}
+
+		public void Detach()
+		{
+			Log.i(TAG, "* Engine::Detach"
+					+ ", lifecycleContextID=" + lifecycleContextID
+					+ ", bResuming=" + bResuming
+					+ ", surfaceCreatedCounter=" + surfaceCreatedCounter
+					+ ", onPauseCounter=" + onPauseCounter
+					+ ", surfaceForTextureView=" + surfaceForTextureView
+			);
+
+			if (view instanceof SurfaceView)
+			{
+				SurfaceView surfaceView = ((SurfaceView)view);
+				surfaceView.getHolder().removeCallback(this);
+			}
+			else if (view instanceof TextureView)
+			{
+				TextureView textureView = ((TextureView)view);
+				textureView.setSurfaceTextureListener(null);
+			}
+
+			((View)view).getViewTreeObserver().removeOnWindowFocusChangeListener(this);
+			//((View)view).getViewTreeObserver().removeOnWindowVisibilityChangeListener(this);
+
+			if (surfaceForTextureView != null)
+			{
+				surfaceForTextureView.release();
+				surfaceForTextureView = null;
+			}
+		}
+
+
+		public SurfaceHolder getHolder() {
+			return type == SurfaceView.class ? ((SurfaceView)view).getHolder() : null;
+		}
+
+		@Override
+		public boolean equals(Object o) {
+			if (this == o) return true;
+			if (o == null || this.type != o.getClass()) return false;
+
+			return view.equals(o);
+		}
+
+		@Override
+		public int hashCode() {
+			return Objects.hash(view, type);
+		}
+	}
+
+	// Volatile Variables...
+	private static int surfaceCreatedCounter = 0;
+	private SimpleContextWrapper activity = null;
+	private static boolean bLibraryLoaded = false;
+	IViewInterface renderView = null;
+
+	private static long maskEngineEvents = 0;
+	private boolean renderSurfaceValid = false;
+	private Surface renderSurface = null;
+	private Surface pendingChangedSurface = null;
+	public static int renderSurfaceViewPosition[] = new int[2];
+	public static int renderSurfaceViewSize[] = new int[2];
+	private static final Object syncLock = new Object();
+	//private GameActivityForMakeAAR gameActivity = null; //TODO remove ref here and use Get() on GameActivity
+	private String OBBFilename = null;
+
+	private boolean bOBBinAPK = true;
+
+	private static boolean bResuming = false;
+	private boolean bAllowConsole = false;
+	private boolean bGameActivityNativeMainIsDone = false;
+
+
+	private static GameActivityForMakeAAR GetOverrideGameActivity()
+	{
+		assert(GameActivity.Get() != null);
+		return (GameActivityForMakeAAR)GameActivity.Get();
+	}
+
+	void DoGameActivityNativeMain(String projectModule, String reasonString)
+	{
+		Log.i(TAG, "DoGameActivityNativeMain is called with reasonString=" + reasonString + ", projectModule=" + projectModule);
+	
+		GetOverrideGameActivity().nativeMain(projectModule);
+
+		bGameActivityNativeMainIsDone = true;
+	}
+
+	public static String[] pendingPakList = null;
+
+	private static int TOUCHTYPE_BEGIN = 0;
+	private static int TOUCHTYPE_MOVE = 1;
+	private static int TOUCHTYPE_END = 2;
+
+	private static int APP_CMD_INPUT_CHANGED = 0;
+	private static int APP_CMD_INIT_WINDOW = 1;
+	private static int APP_CMD_TERM_WINDOW = 2;
+	private static int APP_CMD_WINDOW_RESIZED = 3;
+	private static int APP_CMD_WINDOW_REDRAW_NEEDED = 4;
+	private static int APP_CMD_CONTENT_RECT_CHANGED = 5;
+	private static int APP_CMD_GAINED_FOCUS = 6;
+	private static int APP_CMD_LOST_FOCUS = 7;
+	private static int APP_CMD_CONFIG_CHANGED = 8;
+	private static int APP_CMD_LOW_MEMORY = 9;
+	private static int APP_CMD_START = 10;
+	private static int APP_CMD_RESUME = 11;
+	private static int APP_CMD_SAVE_STATE = 12;
+	private static int APP_CMD_PAUSE = 13;
+	private static int APP_CMD_STOP = 14;
+	private static int APP_CMD_DESTROY = 15;
+
+	private static int logOrder = 0;
+
+	private static boolean _enablePropagateAlpha = false;
+
+	private static int onPauseCounter=0;
+
+	public static Engine Get()
+	{
+		assert (EngineFactory.engine != null);
+		return EngineFactory.engine;
+	}
+
+	public Surface GetRenderSurface()
+	{
+		return renderSurface;
+	}
+
+	public Engine(Activity inActivity, String inOBBFilename, String projectModule, boolean enablePropagateAlpha)
+	{
+		Log.d(TAG, "Engine::Engine() - CONSTRUCTING Engine = " + this);// + ", AARVersion= " + BuildConfig.AAR_VERSION);
+
+		_enablePropagateAlpha = enablePropagateAlpha;
+
+		Log.i(TAG, "Engine - attempting to load engine with inActivity = " + inActivity);
+
+		try {
+			loadLibrary("Unreal");
+		}
+		catch (java.lang.UnsatisfiedLinkError e)
+		{
+			Log.e(TAG, "Engine - failed to load libUE4! UnsatisfiedLinkError: " + e.getMessage());
+			return;
+		}
+
+		// now try the engine
+		if (OBBFilename == null || !inOBBFilename.isEmpty())
+		{
+			Log.i(TAG, "Engine - constructor is setting the OBBFilename: " + inOBBFilename);
+			OBBFilename = inOBBFilename;
+		}
+
+		if (inActivity != null && GameActivity.Get() == null && !OBBFilename.isEmpty())
+		{
+			Log.i(TAG, "Engine - using inActivity to setup new GameActivity: " + inActivity);
+			GameActivityForMakeAAR.Create().setActivity(inActivity, OBBFilename, _enablePropagateAlpha);
+			Log.i(TAG, "** Engine(constructor case1) - calling GetOverrideGameActivity().nativeMain() with : projectModule = " + projectModule);
+			
+			DoGameActivityNativeMain(projectModule, "Engine(case 1)");
+		}
+		else if (GameActivity.Get() != null && !projectModule.isEmpty())
+		{
+			Log.i(TAG, "Engine - inActivity is passed in and projectModule is defined so setup GameActivity nativeMain");
+			// start game thread
+			Log.i(TAG, "** Engine(constructor case2) - calling GetOverrideGameActivity().nativeMain() with : projectModule = " + projectModule);
+			DoGameActivityNativeMain(projectModule, "Engine(case 2)");
+		}
+		else
+		{
+			Log.e(TAG, "Engine - constructor but not creating GameActivity instance!");
+			bGameActivityNativeMainIsDone = false;
+		}
+
+
+		bLibraryLoaded = true;
+
+		Log.i(TAG, "Engine - engine load success");
+	}
+
+	public boolean ApplyPendingSurfaceChanges()
+	{
+		//synchronized(syncLock)
+		{
+			if (pendingChangedSurface == null || pendingChangedSurface == renderSurface) {
+				Log.w(TAG, "logOrder:" + logOrder + " *** Engine::ApplyPendingSurfaceChanges(DO NOTHING pending is NULL, or same as current) - renderSurfaceValid=" + renderSurfaceValid + ", renderSurface=" + renderSurface);
+				return false;
+			}
+
+			final boolean isChanging = (pendingChangedSurface != renderSurface) && pendingChangedSurface != null;
+
+			Log.i(TAG, "logOrder:" + logOrder + " *** Engine::ApplyPendingSurfaceChanges() - bResuming=" + bResuming
+					+ ", isChanging=" + isChanging
+					+ ", onPauseCounter=" + onPauseCounter
+					+ ", renderSurfaceValid=" + renderSurfaceValid
+					+ ", renderSurface = " + renderSurface
+					+ ", pendingChangedSurface=" + pendingChangedSurface
+					+ ", pendingChangedSurface.isValid()=" + pendingChangedSurface.isValid()
+					+ ", renderSurfaceViewSize=" + Arrays.toString(renderSurfaceViewSize)
+			);
+
+			if (isChanging) {
+				++surfaceCreatedCounter;
+			}
+
+			GetOverrideGameActivity().nativeSetSurfaceOverride(pendingChangedSurface, renderSurfaceViewPosition[0], renderSurfaceViewPosition[1]);
+			GetOverrideGameActivity().nativeSetSurfaceViewInfo(renderSurfaceViewSize[0], renderSurfaceViewSize[1]);
+			renderSurface = pendingChangedSurface;
+			renderSurfaceValid = true;
+			pendingChangedSurface = null;
+			GetOverrideGameActivity().nativeAppCommand(APP_CMD_INIT_WINDOW);
+
+
+			return isChanging;
+		}
+	}
+
+	public boolean DoHousekeepingForViewChange(SimpleContextWrapper inActivity, GameActivitySetupInfo gameActivitySetupInfo) {
+
+		if (!bLibraryLoaded)
+		{
+			Log.e(TAG, "** Engine::DoHousekeepingForViewChange(GameActivitySetupInfo) - bLibraryLoaded is false.");
+			return false;
+		}
+
+
+		if (inActivity == null)
+		{
+			Log.e(TAG, "** Engine::DoHousekeepingForViewChange(GameActivitySetupInfo) - inView does not have a valid Activity so do nothing... perhaps should do cleanup on any pre-existing setup?");
+			return false;
+		}
+
+		Log.i(TAG, "Engine::DoHousekeepingForViewChange(GameActivitySetupInfo) ENTERING, current Activity=" + activity + ", inActivity=" + inActivity + ", viewPosition=" + Arrays.toString(gameActivitySetupInfo.ViewPosition) + ", viewSize=" + Arrays.toString(gameActivitySetupInfo.ViewSize) );
+
+
+		boolean bAcitivityNeedsSetup = true;
+
+		// Check for Activity Change...
+		//
+		if (activity != null)
+		{
+			if (activity.equals(inActivity))
+			{
+				Log.i(TAG, "** Engine::DoHousekeepingForViewChange(GameActivitySetupInfo) - inView is using the same acitivity as before so proceed to just updating view...");
+				bAcitivityNeedsSetup = false;
+			}
+			else
+			{
+				Log.i(TAG, "** Engine::DoHousekeepingForViewChange(GameActivitySetupInfo) - previous activity is being changing so should do full cleanup..."
+					+ ", bGameActivityNativeMainIsDone=" + bGameActivityNativeMainIsDone
+					+ ", bCalledStart:" + bCalledStart
+					+ ", startCounter:" + startCounter
+					+ ", bResuming:" + bResuming
+					+ ", onPauseCounter:" + onPauseCounter
+					+ ", renderSurfaceValid:" + renderSurfaceValid
+					+ ", renderSurface:" + renderSurface
+					+ ", renderSurfaceViewPosition:" + Arrays.toString(renderSurfaceViewPosition)
+					+ ", renderSurfaceViewSize:" + Arrays.toString(renderSurfaceViewSize)
+					+ ", surfaceCreatedCounter:" + surfaceCreatedCounter
+					+ ", pendingChangedSurface:" + pendingChangedSurface
+				);
+
+				bAcitivityNeedsSetup = true;
+			}
+		}
+
+
+		boolean bFirstInit = GameActivity.Get() == null;
+
+		if (gameActivitySetupInfo.ViewPosition.length == 2)
+		{
+			renderSurfaceViewPosition = gameActivitySetupInfo.ViewPosition;
+		}
+
+		if (gameActivitySetupInfo.ViewSize.length == 2)
+		{
+			renderSurfaceViewSize = gameActivitySetupInfo.ViewSize;
+		}
+
+		if (bAcitivityNeedsSetup)
+		{
+			//synchronized(syncLock)
+			{
+				if (GameActivity.Get() == null)
+				{
+					Log.i(TAG, "Engine ::DoHousekeepingForViewChange(GameActivitySetupInfo) - using inActivity to setup new GameActivity: " + inActivity + ", bGameActivityNativeMainIsDone=" + bGameActivityNativeMainIsDone);
+
+					GameActivityForMakeAAR.Create().setActivity(inActivity, gameActivitySetupInfo);
+					pendingChangedSurface = gameActivitySetupInfo.renderSurface;
+				}
+				else
+				{
+					Log.i(TAG, "Engine ::DoHousekeepingForViewChange(GameActivitySetupInfo) - using inActivity update GameActivity: " + inActivity + ", bGameActivityNativeMainIsDone=" + bGameActivityNativeMainIsDone);
+
+					GetOverrideGameActivity().setActivity(inActivity, gameActivitySetupInfo);
+					pendingChangedSurface = gameActivitySetupInfo.renderSurface;
+				}
+			}
+			bFirstInit = true;
+		}
+
+		if ((gameActivitySetupInfo.renderSurface != renderSurface) && renderSurface != null)
+		{
+			Log.i(TAG, "Engine::DoHousekeepingForViewChange(GameActivitySetupInfo) - remove old renderSurface =" + renderSurface );
+			releaseVolatileResources("DoHousekeepingForViewChange(GameActivitySetupInfo) - remove old renderSurface");
+		}
+
+		Log.i(TAG, "* Engine::DoHousekeepingForViewChange(GameActivitySetupInfo) - Before recovery checks, bFirstInit=" + bFirstInit + ", bResuming=" + bResuming + ", renderSurfaceValid=" + renderSurfaceValid + ", renderSurface = " + renderSurface);
+
+		++logOrder;
+
+		if (renderSurface != null || surfaceCreatedCounter > 1)
+		{
+			Log.e(TAG, "logOrder:" + logOrder + " ** Engine::onSurfaceTextureAvailable(ERROR SURFACE POSSIBLY ALREADY EXISTS) - surfaceChanged size=" + Arrays.toString(gameActivitySetupInfo.ViewSize) + ", bResuming=" + bResuming + ", renderSurfaceValid=" + renderSurfaceValid + "renderSurface" + renderSurface + ", surfaceCreatedCounter=" + surfaceCreatedCounter );
+		}
+
+		//synchronized(syncLock)
+		{
+			pendingChangedSurface = gameActivitySetupInfo.renderSurface;
+		}
+
+		Log.i(TAG, "* Engine::DoHousekeepingForViewChange(GameActivitySetupInfo) - After recovery checks, bFirstInit=" + bFirstInit + ", bResuming=" + bResuming + ", renderSurfaceValid=" + renderSurfaceValid + ", renderSurface = " + renderSurface);
+
+		activity = inActivity;
+
+		if (bFirstInit && !bResuming)
+		{
+			bResuming = true;
+			Log.i(TAG, "* Engine::DoHousekeepingForViewChange(GameActivitySetupInfo) - check for resuming if renderSurface is not null, bFirstInit=" + bFirstInit
+					+ ", bResuming=" + bResuming
+					+ ", bCalledStart=" + bCalledStart
+					+ ", onPauseCounter=" + onPauseCounter
+					+ ", renderSurfaceValid=" + renderSurfaceValid
+					+ ", renderSurface = " + renderSurface
+			);
+
+			if (!bCalledStart)
+			{
+				QueueStart("DoHousekeepingForViewChange(GameActivitySetupInfo)");
+			}
+
+			if (renderSurface != null)
+			{
+				QueueResume("DoHousekeepingForViewChange(GameActivitySetupInfo)");
+			}
+		}
+
+		Log.i(TAG, "* Engine::DoHousekeepingForViewChange(GameActivitySetupInfo) - setup done!, bFirstInit=" + bFirstInit
+				+ ", bResuming=" + bResuming
+				+ ", bCalledStart=" + bCalledStart
+				+ ", onPauseCounter=" + onPauseCounter
+				+ ", renderSurfaceValid=" + renderSurfaceValid
+				+ ", renderSurface = " + renderSurface
+		);
+
+		return true;
+	}
+
+	public boolean DoHousekeepingForViewChange(Context inContext, IViewInterface view, String OBBFilename, boolean enablePropagateAlpha) {
+
+		if (!bLibraryLoaded)
+		{
+			Log.e(TAG, "** Engine::DoHousekeepingForViewChange(IViewInterface) - bLibraryLoaded is false.");
+			return false;
+		}
+
+		Activity inActivity = SimpleContextWrapper.getActivityOfContext(inContext);
+
+		if (inActivity == null)
+		{
+			Log.e(TAG, "** Engine::DoHousekeepingForViewChange(IViewInterface) - inView does not have a valid Activity so do nothing... perhaps should do cleanup on any pre-existing setup?");
+			return false;
+		}
+
+		boolean bAcitivityNeedsSetup = true;
+
+		// Check for Activity Change...
+		//
+		if (activity != null)
+		{
+			if (activity.equals(inActivity))
+			{
+				Log.w(TAG, "** Engine::DoHousekeepingForViewChange(IViewInterface) - inView is using the same acitivity as before so proceed to just updating view...");
+				bAcitivityNeedsSetup = false;
+			}
+			else
+			{
+				Log.i(TAG, "** Engine::DoHousekeepingForViewChange(IViewInterface) - previous activity is being changing so should do full cleanup..."
+					+ ", bGameActivityNativeMainIsDone=" + bGameActivityNativeMainIsDone
+					+ ", bCalledStart:" + bCalledStart
+					+ ", startCounter:" + startCounter
+					+ ", bResuming:" + bResuming
+					+ ", onPauseCounter:" + onPauseCounter
+					+ ", renderSurfaceValid:" + renderSurfaceValid
+					+ ", renderSurface:" + renderSurface
+					+ ", renderSurfaceViewPosition:" + Arrays.toString(renderSurfaceViewPosition)
+					+ ", renderSurfaceViewSize:" + Arrays.toString(renderSurfaceViewSize)
+					+ ", surfaceCreatedCounter:" + surfaceCreatedCounter
+					+ ", pendingChangedSurface:" + pendingChangedSurface
+					+ ", renderView:" + renderView
+				);
+
+				bAcitivityNeedsSetup = true;
+			}
+		}
+
+
+		boolean bFirstInit = GameActivity.Get() == null;
+
+		if (renderView != null) {
+			renderView.Detach();
+			renderView = null;
+		}
+
+		renderView = view;
+
+		view.GetPosSize(renderSurfaceViewPosition, renderSurfaceViewSize);
+
+		if (bAcitivityNeedsSetup)
+		{
+			//synchronized(syncLock)
+			{
+				if (GameActivity.Get() == null)
+				{
+					Log.i(TAG, "Engine ::DoHousekeepingForViewChange(IViewInterface) - using inActivity to setup new GameActivity: " + inActivity + ", bGameActivityNativeMainIsDone=" + bGameActivityNativeMainIsDone);
+					
+					GameActivityForMakeAAR.Create().setActivity(inActivity, OBBFilename, enablePropagateAlpha);
+				}
+				else
+				{
+					Log.i(TAG, "Engine ::DoHousekeepingForViewChange(IViewInterface) - using inActivity update GameActivity: " + inActivity + ", bGameActivityNativeMainIsDone=" + bGameActivityNativeMainIsDone);
+
+					GetOverrideGameActivity().setActivity(inActivity, OBBFilename, enablePropagateAlpha);
+				}
+			}
+			bFirstInit = true;
+		}
+
+		view.SetupAlpha(enablePropagateAlpha);
+
+		Log.i(TAG, "* Engine::DoHousekeepingForViewChange(IViewInterface) - Before recovery checks, bFirstInit=" + bFirstInit + ", bResuming=" + bResuming + ", renderSurfaceValid=" + renderSurfaceValid + ", renderSurface = " + renderSurface);
+
+		++logOrder;
+
+		Log.i(TAG, "* Engine::DoHousekeepingForViewChange(IViewInterface) - After recovery checks, bFirstInit=" + bFirstInit + ", bResuming=" + bResuming + ", renderSurfaceValid=" + renderSurfaceValid + ", renderSurface = " + renderSurface);
+
+		activity = new SimpleContextWrapper(inActivity);
+
+		if (bFirstInit && !bResuming)
+		{
+			bResuming = true;
+			Log.i(TAG, "* Engine::DoHousekeepingForViewChange(IViewInterface) - check for resuming if renderSurface is not null, bFirstInit=" + bFirstInit
+					+ ", bResuming=" + bResuming
+					+ ", bCalledStart=" + bCalledStart
+					+ ", onPauseCounter=" + onPauseCounter
+					+ ", renderSurfaceValid=" + renderSurfaceValid
+					+ ", renderSurface = " + renderSurface
+			);
+
+			if (!bCalledStart)
+			{
+				QueueStart("DoHousekeepingForViewChange(IViewInterface)");
+			}
+
+			if (renderSurface != null)
+			{
+				QueueResume("DoHousekeepingForViewChange(IViewInterface)");
+			}
+		}
+
+		Log.i(TAG, "* Engine::DoHousekeepingForViewChange(IViewInterface) - setup done!, bFirstInit=" + bFirstInit
+				+ ", bResuming=" + bResuming
+				+ ", bCalledStart=" + bCalledStart
+				+ ", onPauseCounter=" + onPauseCounter
+				+ ", renderSurfaceValid=" + renderSurfaceValid
+				+ ", renderSurface = " + renderSurface
+			);
+
+
+		return true;
+	}
+
+	private boolean VerifyContextID(int contextID, String reasonString)
+	{
+		if (contextID != 0 && lifecycleContextID != 0 && contextID != lifecycleContextID) {
+			Log.e(TAG, "[lifecycleContextID=" + lifecycleContextID + "] contextID:" + contextID + " ABNORMAL VerifyContextID. with reasonString = " + reasonString );
+			return false;
+		}
+
+		return true;
+	}
+
+
+	public boolean AttachExternalRenderSurface(int contextID, Surface externalSurface, int viewPos[], int viewSize[])
+	{
+		++logOrder;
+
+
+		Log.i(TAG, "** Engine::AttachExternalRenderSurface() - bFirstInit"
+			+ ", bGameActivityNativeMainIsDone=" + bGameActivityNativeMainIsDone
+			+ ", contextID:" + contextID
+			+ ", lifecycleContextID:" + lifecycleContextID
+			+ ", bCalledStart:" + bCalledStart
+			+ ", startCounter:" + startCounter
+			+ ", bResuming:" + bResuming
+			+ ", onPauseCounter:" + onPauseCounter
+			+ ", renderSurfaceValid:" + renderSurfaceValid
+			+ ", renderSurface:" + renderSurface
+			+ ", viewPos:" + (viewPos.length == 2 ? String.format(" position[%d,%d]", viewPos[0], viewPos[1]) : viewPos.length)
+			+ ", renderSurfaceViewSize:" + (renderSurfaceViewSize.length == 2 ? String.format(" size[%d,%d]", renderSurfaceViewSize[0], renderSurfaceViewSize[1]) : renderSurfaceViewSize.length)
+			+ ", viewSize:" + (viewSize.length == 2 ? String.format(" size[%d,%d]", viewSize[0], viewSize[1]) : viewSize.length)
+			+ ", surfaceCreatedCounter:" + surfaceCreatedCounter
+			+ ", pendingChangedSurface:" + pendingChangedSurface
+		);
+
+
+		if (!bHasBeenInit || GameActivity.Get() == null || externalSurface == null) {
+			Log.e(TAG, "** Engine::AttachExternalRenderSurface() - ERROR, RETURN NO CHANGE");
+			return false;
+		}
+
+		if (onPauseCounter > 0 || startCounter != 0 || renderSurface != null || pendingChangedSurface != null)
+		{
+			Log.w(TAG, "** Engine::AttachExternalRenderSurface() - CLEANUP previous context"
+				+ ", bHasBeenInit:" + bHasBeenInit
+				+ ", startCounter:" + startCounter
+				+ ", bResuming:" + bResuming
+				+ ", onPauseCounter:" + onPauseCounter
+				+ ", renderSurfaceValid:" + renderSurfaceValid
+				+ ", renderSurface:" + renderSurface
+				+ ", surfaceCreatedCounter:" + surfaceCreatedCounter
+				+ ", pendingChangedSurface:" + pendingChangedSurface
+			);
+
+			if (bHasBeenInit) {
+				if (onPauseCounter<=0) {
+					QueuePause("from AttachExternalRenderSurface onPauseCounter=" + onPauseCounter);
+				}
+				if (startCounter>0)
+				{
+					QueueStop( "from AttachExternalRenderSurface startCounter=" + startCounter);
+				}
+
+				//GetOverrideGameActivity().nativeSetSurfaceViewInfo(1, 1);
+				//GetOverrideGameActivity().nativeSetSurfaceOverride(null, renderSurfaceViewPosition[0], renderSurfaceViewPosition[1]);
+				//GetOverrideGameActivity().nativeAppCommand(APP_CMD_TERM_WINDOW);
+				//GetOverrideGameActivity().nativeAppCommand(APP_CMD_PAUSE);
+				//GetOverrideGameActivity().nativeAppCommand(APP_CMD_STOP);
+				
+//
+//				bCalledStart = false;
+//				startCounter = 0;
+//				bResuming = false;
+//				onPauseCounter = 0;
+//				renderSurfaceValid = false;
+//				renderSurface = null;
+//				surfaceCreatedCounter = 0;
+//				pendingChangedSurface = null;
+//				lifecycleContextID = contextID;
+			}
+		}
+
+
+
+
+		//if (false == VerifyContextID(contextID, "Engine::AttachExternalRenderSurface, contextID=" + contextID)) {
+		//		boolean shouldResume = renderSurface != null || pendingChangedSurface != null || bResuming;
+		//	//onStop(lifecycleContextID, "Engine::AttachExternalRenderSurface");
+		//	releaseVolatileResources(lifecycleContextID, "AttachExternalRenderSurface lifecycleContextID != contextID");
+		//	lifecycleContextID = contextID;
+		//	bResuming = shouldResume;
+		//}
+
+		//else if (renderSurface != null || surfaceCreatedCounter > 0)
+		{
+			//releaseVolatileResources(contextID, "AttachExternalRenderSurface renderSurface != null || surfaceCreatedCounter>0");
+		}
+
+		
+
+		Log.i(TAG, "** Engine::AttachExternalRenderSurface() - CHECK 1A");
+
+
+
+		synchronized(syncLock)
+		{
+			Log.i(TAG, "** Engine::AttachExternalRenderSurface() - CHECK 1B");
+			pendingChangedSurface = externalSurface;
+			bResuming = true;
+			if (viewSize.length==2)
+			{
+				renderSurfaceViewSize[0] = viewSize[0];
+				renderSurfaceViewSize[1] = viewSize[1];
+			}
+			else
+			{
+				Log.w(TAG, "** Engine::AttachExternalRenderSurface() - input viewSize is not defined, so use existing!, renderSurfaceViewSize = " + Arrays.toString(renderSurfaceViewSize));
+			}
+
+			if (viewPos.length==2)
+			{
+				renderSurfaceViewPosition[0] = viewPos[0];
+				renderSurfaceViewPosition[1] = viewPos[1];
+			}
+			else
+			{
+				Log.w(TAG, "** Engine::AttachExternalRenderSurface() - input viewPos is not defined, so use existing!, renderSurfaceViewPosition = " + Arrays.toString(renderSurfaceViewPosition));
+			}
+		}
+
+		QueueStart("AttachExternalRenderSurface");
+
+		Log.i(TAG, "** Engine::AttachExternalRenderSurface() - about to ApplyPendingSurfaceChanges...");
+
+		bResuming |= ApplyPendingSurfaceChanges();
+
+		Log.i(TAG, "** Engine::AttachExternalRenderSurface() - after to ApplyPendingSurfaceChanges...");
+
+		if (bResuming && renderSurface != null)
+		{
+			Log.i(TAG, "** Engine::AttachExternalRenderSurface() - about to QueueResume ...");
+
+			QueueResume( "AttachExternalRenderSurface");
+		}
+
+		Log.i(TAG, "Engine::AttachExternalRenderSurface(DONE)"
+			+ ", startCounter:" + startCounter
+			+ ", bResuming:" + bResuming
+			+ ", onPauseCounter:" + onPauseCounter
+			+ ", renderSurfaceValid:" + renderSurfaceValid
+			+ ", renderSurface:" + renderSurface
+			+ ", surfaceCreatedCounter:" + surfaceCreatedCounter
+			+ ", pendingChangedSurface:" + pendingChangedSurface
+		);
+
+		return bResuming;
+	}
+
+	@UiThread
+public static synchronized void QueuePause(String reasonString)
+	{
+		if (bResuming && onPauseCounter <= 0) {
+			Log.e(TAG, "QueuePause is CALLED with lifecycleContextID=" + lifecycleContextID + ", bResuming=" + bResuming);
+			QueueResume("QueuePause but was resuming and reasonString=" + reasonString);
+			return;
+		}
+
+		if (GameActivity.Get() == null) {
+			Log.e(TAG, "QueuePause is REJECTED cause GameActivity is not valid with lifecycleContextID=" + lifecycleContextID + ", bResuming=" + bResuming + ", reasonString=" + reasonString);
+			return;
+		}
+
+		GetOverrideGameActivity().GetCurrentActivityContext().runOnUiThread(() ->
+				{
+					//synchronized(syncLock)
+					{
+						Log.i(TAG, "QueuePause is calling engine.onPause with lifecycleContextID=" + lifecycleContextID + ", bResuming=" + bResuming);
+						Engine.Get().onPause(lifecycleContextID, reasonString);
+					}
+				}
+		);
+	}
+
+	@UiThread
+	public static synchronized void QueueResume(String reasonString)
+	{
+		if (!GameActivityForMakeAAR.isValidGameActivity())
+		{
+			return;
+		}
+
+		if ( Get().pendingChangedSurface == null && Get().renderSurface == null)
+		{
+			Log.e(TAG, "QueueResume is IGNORED, lifecycleContextID=" + lifecycleContextID + ", reasonString=" + reasonString + ", bResuming= " + bResuming + ", pendingChangedSurface=" + Get().pendingChangedSurface);
+			return;
+		}
+
+		GetOverrideGameActivity().GetCurrentActivityContext().runOnUiThread(() ->
+			{
+				//synchronized (syncLock)
+				{
+					Log.i(TAG, "QueueResume is calling engine.onResume with lifecycleContextID=" + lifecycleContextID + ", reasonString=" + reasonString);
+					
+					Engine.Get().onResume(lifecycleContextID, "QueueResume with reasonString=" + reasonString);
+				}
+			}
+		);
+	}
+
+	@UiThread
+	public static synchronized void QueueStart(String reasonString)
+	{
+		if (!GameActivityForMakeAAR.isValidGameActivity())
+		{
+			Log.e(TAG, "QueueStart called with invalid GameActivity!, lifecycleContextID=" + lifecycleContextID + ", reasonString=" + reasonString);
+			return;
+		}
+
+		GetOverrideGameActivity().GetCurrentActivityContext().runOnUiThread(()->
+			{
+				//synchronized (syncLock)
+				{
+					Log.i(TAG, "QueueStart is calling engine.onStart() with lifecycleContextID=" + lifecycleContextID + ", reasonString=" + reasonString);
+	
+					Engine.Get().onStart(lifecycleContextID, "QueueStart with reasonString=" + reasonString);
+				}
+			}
+		);
+	}
+
+	@UiThread
+	public static synchronized void QueueStop(String reasonString)
+	{
+		if (false == GameActivityForMakeAAR.isValidGameActivity())
+		{
+			Log.e(TAG, "QueueStop called with invalid GameActivity!, lifecycleContextID=" + lifecycleContextID + ", reasonString=" + reasonString);
+			return;
+		}
+
+		GetOverrideGameActivity().GetCurrentActivityContext().runOnUiThread(()->
+			{
+				//synchronized(syncLock)
+				{
+					Log.i(TAG, "QueueStop is calling engine.onStop with lifecycleContextID=" + lifecycleContextID + ", reasonString=" + reasonString);
+					Engine.Get().onStop(lifecycleContextID, "QueueStop with reasonString=" + reasonString);
+				}
+	
+			}
+		);
+	}
+
+	boolean bHasBeenInit = false;
+
+	public boolean IsInitialized() { return  bHasBeenInit; }
+
+
+	public boolean Init(Context context, GameActivitySetupInfo gameActivitySetupInfo, String commandline, String projectModule, String[] mountPAKs)
+	{
+		incCreateContextID();
+
+		++logOrder;
+		Log.i(TAG, "logOrder:" + logOrder + " * Engine::Init(context, gameActivitySetupInfo) -" + ", bLibraryLoaded=" + bLibraryLoaded + ", bResuming=" + bResuming + ", renderSurfaceValid=" + renderSurfaceValid + ", context = " + context);
+		if (renderSurfaceValid || renderSurface != null)
+		{
+			releaseVolatileResources("Engine::Init(context, gameActivitySetupInfo) called, so remove old resources");
+		}
+		if (!bLibraryLoaded)
+		{
+			return false;
+		}
+
+		boolean bFirstInit = (GameActivity.Get() == null || renderSurface == null);
+
+		if (!DoHousekeepingForViewChange(new SimpleContextWrapper(context), gameActivitySetupInfo))
+		{
+			return false;
+		}
+
+		bHasBeenInit = true;
+
+		if (bFirstInit)
+		{
+			Log.i(TAG, "** Engine::Init(context, gameActivitySetupInfo) - bFirstInit"
+					+ ", bGameActivityNativeMainIsDone=" + bGameActivityNativeMainIsDone
+					+ ", bCalledStart:" + bCalledStart
+					+ ", startCounter:" + startCounter
+					+ ", bResuming:" + bResuming
+					+ ", onPauseCounter:" + onPauseCounter
+					+ ", renderSurfaceValid:" + renderSurfaceValid
+					+ ", renderSurface:" + renderSurface
+					+ ", pendingChangedSurface:" + pendingChangedSurface
+					+ ", commandline:" + commandline
+					+ ", projectModule:" + projectModule
+			);
+
+			if (commandline.isEmpty() || !commandline.startsWith("../../"))
+			{
+				commandline = "../../../" + projectModule + "/" + projectModule + ".uproject " + commandline;
+				Log.i(TAG, "** Engine::Init(External View) - commandline updated: " + commandline);
+			}
+
+			GetOverrideGameActivity().setCommandline(commandline);
+
+			if (mountPAKs.length > 0) {
+				pendingPakList = mountPAKs.clone();
+			}
+
+			// start game thread
+			if (!bGameActivityNativeMainIsDone)
+			{
+				Log.i(TAG, "** Engine::Init(context, gameActivitySetupInfo) - calling GetOverrideGameActivity().nativeMain() with : projectModule = " + projectModule);
+				DoGameActivityNativeMain(projectModule, "Init(context, gameActivitySetupInfo)");
+			}
+
+			AndroidThunkJava_Engine_ReceiveEvent(EVENTTYPE_ACTIVITY_HAS_CHANGED, "", 0, 0, 0.0f);
+		}
+
+
+		if (!bCalledStart)
+		{
+			QueueStart("Init(context, gameActivitySetupInfo)");
+			QueueResume("Init(context, gameActivitySetupInfo) because started");
+		}
+
+		if (renderSurface != null || pendingChangedSurface != null)
+		{
+			QueueResume("Init(context, gameActivitySetupInfo) refresh surface");
+		}
+
+		return bResuming;
+	}
+
+	public boolean Init(SurfaceView view, String commandline, String projectModule, String[] mountPAKs)
+	{
+		return Init(view.getContext(), new ViewInterface<SurfaceView>(view), this.OBBFilename, commandline, projectModule, _enablePropagateAlpha, mountPAKs);
+	}
+
+	public boolean Init(TextureView view, String commandline, String projectModule, String[] mountPAKs)
+	{
+		return Init(view.getContext(), new ViewInterface<TextureView>(view), this.OBBFilename, commandline, projectModule, _enablePropagateAlpha, mountPAKs);
+	}
+
+	public boolean Init(Context inContext, IViewInterface view, String OBBFilename, String commandline, String projectModule, boolean enablePropagateAlpha, String[] mountPAKs)
+	{
+		incCreateContextID();
+
+		++logOrder;
+		Log.i(TAG, "logOrder:" + logOrder + " * Engine::Init(IViewInterface) -" + ", bLibraryLoaded=" + bLibraryLoaded + ", bResuming=" + bResuming + ", renderSurfaceValid=" + renderSurfaceValid + ", view = " + view);
+		if (renderSurfaceValid || renderSurface != null)
+		{
+			releaseVolatileResources( "Engine::Init(IViewInterface) called so remove old resources");
+		}
+		if (!bLibraryLoaded)
+		{
+			return false;
+		}
+
+		boolean bFirstInit = (GameActivity.Get() == null || renderSurface == null);
+
+
+		if (!DoHousekeepingForViewChange(inContext, view, OBBFilename, enablePropagateAlpha ))
+		{
+			return false;
+		}
+
+		bHasBeenInit = true;
+
+		if (bFirstInit)
+		{
+			Log.i(TAG, "** Engine::Init(ViewInterface) - bFirstInit"
+					+ ", bGameActivityNativeMainIsDone=" + bGameActivityNativeMainIsDone
+					+ ", bCalledStart:" + bCalledStart
+					+ ", startCounter:" + startCounter
+					+ ", bResuming:" + bResuming
+					+ ", onPauseCounter:" + onPauseCounter
+					+ ", renderSurfaceValid:" + renderSurfaceValid
+					+ ", renderSurface:" + renderSurface
+					+ ", commandline:" + commandline
+					+ ", projectModule:" + projectModule
+			);
+
+			if (commandline.isEmpty() || !commandline.startsWith("../../"))
+			{
+				commandline = "../../../" + projectModule + "/" + projectModule + ".uproject " + commandline;
+				Log.i(TAG, "** Engine::Init(ViewInterface) - commandline updated: " + commandline);
+			}
+
+			GetOverrideGameActivity().setCommandline(commandline);
+
+			if (mountPAKs.length > 0) {
+				pendingPakList = mountPAKs.clone();
+			}
+
+			// start game thread
+			if (!bGameActivityNativeMainIsDone)
+			{
+				Log.i(TAG, "** Engine::Init(IViewInterface) - calling GetOverrideGameActivity().nativeMain() with : projectModule = " + projectModule);
+				DoGameActivityNativeMain(projectModule, "Init(IViewInterface)");
+			}
+
+			AndroidThunkJava_Engine_ReceiveEvent(EVENTTYPE_ACTIVITY_HAS_CHANGED, "", 0, 0, 0.0f);
+		}
+
+
+		if (!bCalledStart)
+		{
+			QueueStart("Init(ViewInterface)");
+		}
+
+		if (pendingChangedSurface != null)
+		{
+			QueueResume("Init(ViewInterface)");
+		}
+
+		return bResuming;
+	}
+
+	public String getObbVersion(boolean bFull)
+	{
+		String rawComment = GetOverrideGameActivity().nativeGetObbComment();
+		if (bFull)
+		{
+			return rawComment;
+		}
+		int separatorIndex = rawComment.indexOf(":");
+		return (separatorIndex >= 0) ? rawComment.substring(0, separatorIndex) : "0.0.0";
+	}
+
+	private int getResourceId(String VariableName, String ResourceName, String PackageName)
+	{
+		Log.i(TAG, "** Engine::getResourceId(...), VariableName=" + VariableName + ", ResourceName=" + ResourceName + ", PackageName=" + PackageName );
+		
+		try {
+			return GameActivityForMakeAAR.GetCurrentActivityContext().getApplicationContext().getResources().getIdentifier(VariableName, ResourceName, PackageName);
+		}
+		catch (Exception e) {
+			e.printStackTrace();
+			return -1;
+		}
+	}
+
+	public String getAARVersion()
+	{
+		if (GameActivityForMakeAAR.isValidGameActivity() == false)
+		{
+			return "0.0.0";
+		}
+
+		try {
+			int resourceId = getResourceId("UnrealEngineAARVersion", "string", activity.getPackageName());
+			return (resourceId < 1) ? "0.0.0" : activity.getString(resourceId);
+		}
+		catch (Exception e)
+		{
+			e.printStackTrace();
+			return "0.0.-1";
+		}
+	}
+
+	public void allowConsole(boolean bEnable)
+	{
+		bAllowConsole = bEnable;
+	}
+
+	public void sendConsoleCommand(String command)
+	{
+		if (GetOverrideGameActivity() != null) {
+			GetOverrideGameActivity().nativeConsoleCommand(command);
+		}
+	}
+
+	private boolean bCalledStart = false;
+	static int startCounter = 0;
+	static int lifecycleContextID = 0;
+	static int lifecycleActivityContextCounter = 0;
+
+	public static int getCurrentContextID()
+	{
+		return lifecycleContextID;
+	}
+
+	//public int getCreateContext()
+	//{
+	//	return lifecycleContextID++;
+	//}
+
+	public void incCreateContextID()
+	{
+		++logOrder;
+		Log.w(TAG, "[lifecycleContextID=" + lifecycleContextID + "] logOrder:" + logOrder + " ** Engine::incCreateContextID(), "
+				+ ", lifecycleActivityContextCounter=" + lifecycleActivityContextCounter
+				+ ", bResuming=" + bResuming
+				+ ", renderSurfaceValid=" + renderSurfaceValid
+				+ ", surfaceCreatedCounter = " + surfaceCreatedCounter
+				+ ", renderSurface = " + renderSurface
+		);
+		if (startCounter > 0 && GameActivityForMakeAAR.isValidGameActivity())
+		{
+			onStop(lifecycleContextID, "incCreateContextID");
+		}
+		lifecycleContextID = (lifecycleContextID%100) + (++lifecycleActivityContextCounter) * 100;
+	}
+
+	private void onPause(int contextID, String reasonString)
+	{
+		if (contextID != 0 && lifecycleContextID != 0 && contextID != lifecycleContextID) {
+			Log.e(TAG, "[lifecycleContextID=" + lifecycleContextID + "] contextID:" + contextID + " ABNORMAL Engine::onPause(contextID) - onPause called. with contextID change = " + contextID + ", reasonString=" + reasonString);
+			return;
+		}
+
+		if (!bHasBeenInit) {
+			Log.e(TAG, "[lifecycleContextID=" + lifecycleContextID + "] logOrder:" + logOrder + " * Engine::onPause - BUT bHasBeenInit is FALSE, DO NOTHING. renderSurfaceValid: " + renderSurfaceValid + ", bResuming: " + bResuming + ", bCalledStart = " + bCalledStart + ", onPauseCounter=" + onPauseCounter + ", renderSurface = " + renderSurface + ", reasonString=" + reasonString);
+			return;
+		}
+
+		Log.i(TAG, "[lifecycleContextID=" + lifecycleContextID + "] logOrder:" + logOrder + " * Engine::onPause - renderSurfaceValid: " + renderSurfaceValid + ", bResuming: " + bResuming + ", bCalledStart = " + bCalledStart + ", onPauseCounter=" + onPauseCounter + ", renderSurface = " + renderSurface + ", pendingChangedSurface=" + pendingChangedSurface + ", reasonString=" + reasonString);
+
+		//if (_gameActivityInstance != null) {
+		//	GetOverrideGameActivity().onPause();
+		//}
+
+		if (onPauseCounter > 0 || bResuming) {
+			Log.w(TAG, "** Engine::onPause - Already Paused! onPauseCounter=" + onPauseCounter + ", bResuming=" + bResuming + ", reasonString=" + reasonString);
+			return;
+		}
+
+		++onPauseCounter;
+		Log.i(TAG, "** Engine::onPause - DOING a full Pause!" + ", reasonString=" + reasonString);
+
+		if (renderSurface != null) {
+			--surfaceCreatedCounter;
+		}
+
+		if (renderSurface != null || pendingChangedSurface != null)
+		{
+			GetOverrideGameActivity().nativeSetSurfaceOverride(null, renderSurfaceViewPosition[0], renderSurfaceViewPosition[1]);
+			//pendingChangedSurface = pendingChangedSurface != null ? pendingChangedSurface : renderSurface;
+			renderSurface = null;
+			renderSurfaceValid = false;
+			bResuming |= pendingChangedSurface != null;
+			Log.i(TAG, "** Engine::onPause - set SurfaceOverride to null and updated render surfaces" + ", pendingChangedSurface=" + pendingChangedSurface + ", reasonString=" + reasonString);
+		}
+
+		GetOverrideGameActivity().nativeAppCommand(APP_CMD_LOST_FOCUS);
+		GetOverrideGameActivity().nativeAppCommand(APP_CMD_PAUSE);
+		//if (pendingChangedSurface != null)
+		//	GetOverrideGameActivity().nativeAppCommand(APP_CMD_TERM_WINDOW);
+
+		//releaseVolatileResources("onPause called with reason=" + reasonString);
+	}
+
+	private void onResume(int contextID, String reasonString)
+	{
+		if (contextID != 0 && lifecycleContextID != 0 && contextID != lifecycleContextID) {
+			Log.e(TAG, "[lifecycleContextID=" + lifecycleContextID + "] contextID:" + contextID + " ABNORMAL Engine::onResume - onResume called. with contextID change = " + contextID + ", reasonString=" + reasonString);
+			return;
+		}
+		++logOrder;
+		Log.i(TAG, "[lifecycleContextID=" + lifecycleContextID + "] logOrder:" + logOrder + " * Engine::onResume - onResume called. bCalledStart was = " + bCalledStart
+			+ ", bHasBeenInit = " + bHasBeenInit
+			+ ", maskEngineEvents= " + maskEngineEvents
+			+ ", bGameActivityNativeMainIsDone = " + bGameActivityNativeMainIsDone
+			+ ", prev startCounter = " + startCounter
+			+ ", surfaceCreatedCounter = " + surfaceCreatedCounter
+			+ ", bResuming was = " + bResuming
+			+ ", onPauseCounter=" + onPauseCounter
+			+ ", renderSurfaceValid: " + renderSurfaceValid
+			+ ", renderSurface = " + renderSurface
+		);
+
+		if (bCalledStart == false) {
+			if (surfaceCreatedCounter > 0 && renderSurface != null) {
+				Log.w(TAG, "** Engine::onResume - FORCING onstart from onresume because bCalledStart is false! bResuming=" + bResuming + ", onPauseCounter=" + onPauseCounter + ", width = " + renderSurfaceViewSize[0] + ", height = " + renderSurfaceViewSize[1]);
+
+				onStart(lifecycleContextID, reasonString + "->onResume");
+
+			} else {
+				Log.w(TAG, "** Engine::onResume - SKIPPING because bCalledStart is false! bResuming=" + bResuming + ", onPauseCounter=" + onPauseCounter + ", width = " + renderSurfaceViewSize[0] + ", height = " + renderSurfaceViewSize[1]);
+				return;
+			}
+		}
+
+//		if (GameActivity.Get() != null) {
+//			GameActivity.Get().onResume();
+//		}
+
+
+		bResuming |= ApplyPendingSurfaceChanges();
+
+		if (renderSurface == null || !renderSurfaceValid)
+		{
+			Log.i(TAG, "** Engine::onResume - no valid surface to resume... bResuming=" + bResuming + ", surfaceCreatedCounter" + surfaceCreatedCounter + ", onPauseCounter=" + onPauseCounter + ", renderSurfaceValid = " + renderSurfaceValid + ", renderSurface = " + renderSurface);
+
+			return;
+		}
+
+		Log.i(TAG, "** Engine::onResume - CHECKING resume cases... bResuming=" + bResuming + ", onPauseCounter=" + onPauseCounter + ", width = " + renderSurfaceViewSize[0] + ", height = " + renderSurfaceViewSize[1]);
+
+		onPauseCounter = onPauseCounter > 0 ? --onPauseCounter : onPauseCounter;
+
+		GetOverrideGameActivity().nativeAppCommand(APP_CMD_RESUME);
+
+		bResuming |= onPauseCounter > 0;
+
+		if (!bResuming && (bCalledStart && surfaceCreatedCounter>0 && renderSurfaceValid && renderSurface != null))
+		{
+			Log.i(TAG, "** Engine::onResume - setup renderSurface because we already are in a start state! width = " + renderSurfaceViewSize[0] + ", height = " + renderSurfaceViewSize[1]);
+
+			GetOverrideGameActivity().nativeSetSurfaceOverride(renderSurface, renderSurfaceViewPosition[0], renderSurfaceViewPosition[1]);
+			GetOverrideGameActivity().nativeSetSurfaceViewInfo(renderSurfaceViewSize[0], renderSurfaceViewSize[1]);
+			bResuming = true;
+		}
+
+		if (bResuming && renderSurface != null) {
+			Log.i(TAG, "** Engine::onResume - DOING full Resume! width = " + renderSurfaceViewSize[0] + ", height = " + renderSurfaceViewSize[1] + ", maskEngineEvents= " + maskEngineEvents);
+
+			GetOverrideGameActivity().nativeAppCommand(APP_CMD_GAINED_FOCUS);
+			GetOverrideGameActivity().nativeResumeMainInit();
+			bResuming = false;
+		}
+		else if (renderSurface != null)
+		{
+			Log.e(TAG, "***- Engine::onResume - UNHANDLED case for resume! renderSurfaceValid = " + renderSurfaceValid + ", renderSurface = " + renderSurface);
+			GetOverrideGameActivity().nativeSetSurfaceOverride(renderSurface, renderSurfaceViewPosition[0], renderSurfaceViewPosition[1]);
+			GetOverrideGameActivity().nativeSetSurfaceViewInfo(renderSurfaceViewSize[0], renderSurfaceViewSize[1]);
+			GetOverrideGameActivity().nativeAppCommand(APP_CMD_GAINED_FOCUS);
+			GetOverrideGameActivity().nativeResumeMainInit();
+		}
+		else
+		{
+			Log.e(TAG, "***- Engine::onResume - EMPTY unhandled case for resume! renderSurfaceValid = " + renderSurfaceValid + ", renderSurface = " + renderSurface);
+		}
+	}
+
+	private void onStart(int contextID, String reasonString)
+	{
+		if (contextID != 0 && lifecycleContextID != 0 && contextID != lifecycleContextID) {
+			Log.e(TAG, "[lifecycleContextID=" + lifecycleContextID + "] contextID:" + contextID + " ABNORMAL Engine::onStart(contextID) - onStart called. with contextID change = " + contextID + ", reasonString=" + reasonString);
+			return;
+		}
+
+		++logOrder;
+		Log.i(TAG, "[lifecycleContextID=" + lifecycleContextID + "] logOrder:" + logOrder + " * Engine::onStart - onStart called. bCalledStart was = " + bCalledStart
+			+ ", bHasBeenInit = " + bHasBeenInit
+			+ ", prev startCounter = " + startCounter
+			+ ", bResuming was = " + bResuming
+			+ ", onPauseCounter=" + onPauseCounter
+			+ ", renderSurfaceValid: " + renderSurfaceValid
+			+ ", renderSurface = " + renderSurface
+		);
+
+		if (!bHasBeenInit) {
+		}
+
+		++startCounter;
+		if (!bCalledStart) {
+			bCalledStart = true;
+			GetOverrideGameActivity().nativeAppCommand(APP_CMD_START);
+		}
+	}
+
+	private void onStop(int contextID, String reasonString)
+	{
+		if (contextID != 0 && lifecycleContextID != 0 && contextID != lifecycleContextID) {
+			Log.e(TAG, "[lifecycleContextID=" + lifecycleContextID + "] contextID:" + contextID + " ABNORMAL Engine::onStop(contextID) - onStop called. with contextID change = " + contextID + ", reasonString=" + reasonString);
+			return;
+		}
+
+		++logOrder;
+		Log.i(TAG, "[lifecycleContextID=" + lifecycleContextID + "] logOrder:" + logOrder + " * Engine::onStop - onStop called. bCalledStart was = " + bCalledStart
+				+ ", reasonString=" + reasonString
+				+ ", bHasBeenInit = " + bHasBeenInit
+				+ ", prev startCounter = " + startCounter
+				+ ", surfaceCreatedCounter = " + surfaceCreatedCounter
+				+ ", bResuming was = " + bResuming
+				+ ", onPauseCounter=" + onPauseCounter
+				+ ", renderSurfaceValid: " + renderSurfaceValid
+				+ ", renderSurface = " + renderSurface
+		);
+
+		if (onPauseCounter <= 0 || renderSurface != null) {
+			onPause(contextID, reasonString + "from onStop()");
+		}
+
+		//releaseVolatileResources("onStop() called so release all resources");
+		//synchronized(syncLock)
+		{
+			pendingChangedSurface = null;
+		}
+		
+		if (startCounter>0) {
+			--startCounter;
+		}
+
+		if (bCalledStart)
+		{
+			GetOverrideGameActivity().nativeAppCommand(APP_CMD_STOP);
+			GetOverrideGameActivity().nativeAppCommand(APP_CMD_TERM_WINDOW);
+
+			bCalledStart = false;
+			startCounter = 0;
+			onPauseCounter = 0;
+			lifecycleContextID = 0;
+		}
+
+		Log.i(TAG, "[lifecycleContextID=" + lifecycleContextID + "] logOrder:" + logOrder + " * Engine::onStop - onStop COMPLETED. bCalledStart was = " + bCalledStart
+			+ ", reasonString=" + reasonString
+			+ ", bHasBeenInit = " + bHasBeenInit
+			+ ", prev startCounter = " + startCounter
+			+ ", surfaceCreatedCounter = " + surfaceCreatedCounter
+			+ ", bResuming was = " + bResuming
+			+ ", onPauseCounter=" + onPauseCounter
+			+ ", renderSurfaceValid: " + renderSurfaceValid
+			+ ", renderSurface = " + renderSurface
+		);
+	}
+
+	private void ReleaseRenderSurfaces(int contextID, String reasonString)
+	{
+		Log.i(TAG, "*** Engine::ReleaseRenderSurfaces() : GameActivityForMakeAAR.isValidGameActivity() = " + GameActivityForMakeAAR.isValidGameActivity()
+			+ ", bGameActivityNativeMainIsDone=" + bGameActivityNativeMainIsDone
+			+ ", contextID:" + contextID
+			+ ", lifecycleContextID:" + lifecycleContextID
+			+ ", reasonString:" + reasonString
+			+ ", bCalledStart:" + bCalledStart
+			+ ", startCounter:" + startCounter
+			+ ", bResuming:" + bResuming
+			+ ", onPauseCounter:" + onPauseCounter
+			+ ", renderSurfaceValid:" + renderSurfaceValid
+			+ ", renderSurface:" + renderSurface
+			+ ", surfaceCreatedCounter:" + surfaceCreatedCounter
+			+ ", pendingChangedSurface:" + pendingChangedSurface
+
+		);
+
+		if (GameActivityForMakeAAR.isValidGameActivity() && (onPauseCounter <= 0 || renderSurface != null || pendingChangedSurface != null)) {
+
+			onPause(contextID, reasonString + "->ReleaseRenderSurfaces");
+
+
+			//if (renderSurface != null)
+			{
+				//GetOverrideGameActivity().nativeAppCommand(APP_CMD_PAUSE);
+				//GetOverrideGameActivity().nativeAppCommand(APP_CMD_STOP);
+
+				
+
+				//if (onPauseCounter > 0) {
+				//	QueuePause("ReleaseRenderSurfaces");
+				//}
+
+				//if (startCounter > 0) {
+				//	QueueStop("ReleaseRenderSurfaces");
+				//}
+			}
+			//else if (!CheckEvent(EVENTTYPE_INIT))
+			//{
+			//	GetOverrideGameActivity().nativeAppCommand(APP_CMD_PAUSE);
+			//}
+
+			if (renderSurfaceValid && (startCounter > 0 || onPauseCounter > 0)) {
+				bResuming = true;
+			}
+				
+			if (renderSurface != null) {
+				pendingChangedSurface = renderSurface;
+				bResuming = true;
+			}
+			renderSurface = null;
+			renderSurfaceValid = false;
+		}
+
+	}
+
+	private void releaseVolatileResources(String reasonString)
+	{
+		releaseVolatileResources(lifecycleContextID, reasonString);
+	}
+
+	private void releaseVolatileResources(int contextID, String reasonString)
+	{
+		Log.i(TAG, "logOrder:" + logOrder + " * Engine::releaseVolatileResources() - bCalledStart was = " + bCalledStart + ", contextID=" + contextID + ", reasonString=" + reasonString + ", bResuming was = " + bResuming + ", onPauseCounter=" + onPauseCounter + ", renderSurfaceValid: " + renderSurfaceValid + ", renderSurface = " + renderSurface );
+
+		ReleaseRenderSurfaces(contextID, reasonString);
+	}
+
+	public void onDestroy(int contextID)
+	{
+		onDestroy(contextID, "External");
+	}
+
+	private void onDestroy(int customContextID, String reasonString)
+	{
+		++logOrder;
+		Log.i(TAG, "[lifecycleContextID=" + lifecycleContextID + "] logOrder:" + logOrder + " * Engine::onDestroy() - onDestroy called. bCalledStart was = " + bCalledStart + ", bResuming was = " + bResuming + ", onPauseCounter=" + onPauseCounter + ", renderSurfaceValid: " + renderSurfaceValid + ", renderSurface = " + renderSurface );
+
+		//onPause(customContextID, "onDestroy called");
+		ReleaseRenderSurfaces(customContextID, reasonString + "->onDestroy called");
+
+		if (customContextID == 543210 || !CheckEvent(EVENTTYPE_INIT))
+		{
+
+			// we don't want to call APP_CMD_DESTROY unless we want to completely terminate the unreal engine instance.
+			GetOverrideGameActivity().nativeAppCommand(APP_CMD_DESTROY);
+
+		}
+
+		//if (renderSurface != null)
+		{
+			GetOverrideGameActivity().nativeAppCommand(APP_CMD_TERM_WINDOW);
+		}
+
+		if (renderView != null) {
+			renderView.Detach();
+			renderView = null;
+		}
+
+		maskEngineEvents = 0;
+		bHasBeenInit = false;
+		bResuming = false;
+
+
+		lifecycleContextID += 25;
+		if (onPauseCounter>0)
+		{
+			Log.e(TAG, "onDestroy is called with onPauseCounter>0 which should not happen for healthy lifecycle");
+		}
+
+		Log.w(TAG, "*** Engine::onDestroy() -> AndroidThunkJava_RestartApplication()");
+		GetOverrideGameActivity().AndroidThunkJava_RestartApplication("");
+
+		Log.w(TAG, "*** Engine::onDestroy() set statics to defaults!");
+
+		surfaceCreatedCounter = 0;
+		activity = null;
+		//bLibraryLoaded = false;
+		renderView = null;
+
+		maskEngineEvents = 0;
+		renderSurfaceValid = false;
+		renderSurface = null;
+		pendingChangedSurface = null;
+		//renderSurfaceViewPosition[] = new int[2];
+		//renderSurfaceViewSize[] = new int[2];
+		//syncLock = new Object();
+		//private GameActivityForMakeAAR gameActivity = null; //TODO remove ref here and use Get() on GameActivity
+		//private String OBBFilename = null;
+		//private boolean bOBBinAPK = true;
+
+		bResuming = false;
+		//private boolean bAllowConsole = false;
+		bGameActivityNativeMainIsDone = false;
+
+
+	
+		Log.i(TAG, "[lifecycleContextID=" + lifecycleContextID + "] logOrder:" + logOrder + " * Engine::onDestroy() - onDestroy completed. bCalledStart was = " + bCalledStart + ", bResuming was = " + bResuming + ", onPauseCounter=" + onPauseCounter + ", renderSurfaceValid: " + renderSurfaceValid + ", renderSurface = " + renderSurface );
+
+
+	
+	}
+
+	public void onConfigurationChanged(Configuration newConfig)
+	{
+		// forward the orientation
+		boolean bPortrait = newConfig.orientation == Configuration.ORIENTATION_PORTRAIT;
+		GetOverrideGameActivity().nativeOnOrientationChanged(newConfig.orientation);
+	}
+
+	public void onActivityResult(int requestCode, int resultCode, Intent data)
+	{
+		++logOrder;
+		Log.v(TAG, "logOrder:" + logOrder + " * Engine::onActivityResult() - bCalledStart was = " + bCalledStart + ", bResuming was = " + bResuming + ", onPauseCounter=" + onPauseCounter + ", renderSurfaceValid: " + renderSurfaceValid + ", renderSurface = " + renderSurface );
+
+		// only used by FOnlineSubsystemGooglePlay so we changed signature to pass the ContextWrapper instead of activity here because we no longer want to hold onto any activity references...
+		GetOverrideGameActivity().nativeOnActivityResult(GameActivity.Get(), requestCode, resultCode, data);
+//		onResume();
+	}
+
+	public boolean onTouchEvent(MotionEvent event)
+	{
+		int pointerCount = event.getPointerCount();
+		if (pointerCount < 1 || !GameActivityForMakeAAR.isValidGameActivity())
+		{
+			return false;
+		}
+		int device = event.getDeviceId();
+
+		int action = 0;
+		boolean bProcessed = false;
+		for (int index = 0; index < pointerCount; index++) {
+			int pointerId = event.getPointerId(index);
+			int x = (int)event.getX();
+			int y = (int)event.getY();
+			switch (event.getAction()) {
+				case MotionEvent.ACTION_DOWN:
+				case MotionEvent.ACTION_POINTER_DOWN:
+					action = TOUCHTYPE_BEGIN;
+					break;
+				case MotionEvent.ACTION_MOVE:
+					action = TOUCHTYPE_MOVE;
+					break;
+				case MotionEvent.ACTION_UP:
+				case MotionEvent.ACTION_POINTER_UP:
+				case MotionEvent.ACTION_CANCEL:
+				case MotionEvent.ACTION_OUTSIDE:
+					action = TOUCHTYPE_END;
+					break;
+				case MotionEvent.ACTION_SCROLL:
+				case MotionEvent.ACTION_HOVER_ENTER:
+				case MotionEvent.ACTION_HOVER_MOVE:
+				case MotionEvent.ACTION_HOVER_EXIT:
+					continue;
+				default:
+					continue;
+			}
+
+			if (GetOverrideGameActivity().nativeInputTouch(device, action, pointerId, x, y) == 1) {
+				bProcessed = true;
+			}
+		}
+
+		if (bAllowConsole && (pointerCount >= 4) && (action == TOUCHTYPE_BEGIN)) {
+			showConsole();
+		}
+
+		return bProcessed;
+	}
+
+	public void showConsole()
+	{
+		if (GameActivityForMakeAAR.isValidGameActivity()) {
+			GetOverrideGameActivity().AndroidThunkJava_ShowConsoleWindow(null);
+		}
+	}
+
+	public static void registerEventCallback(IEventCallback callback)
+	{
+		if (null == callback) {
+			mCallbackRef = new WeakReference<>(IEventCallback.NO_OP);
+		} else {
+			mCallbackRef = new WeakReference<>(callback);
+		}
+	}
+
+	public static void AndroidThunkJava_Engine_ReceiveEvent(int event, String param1, int param2, int param3, float param4)
+	{
+		if (event == EVENTTYPE_INIT) {
+			String AARVersion = Engine.Get().getAARVersion();
+			String OBBVersion = Engine.Get().getObbVersion(false);
+
+			Log.d(TAG, "*** Engine::ReceiveEvent - INIT - AARVersion=" + AARVersion + ", OBBVersion=" + OBBVersion);
+			if (!AARVersion.equals(OBBVersion)) {
+				Log.w(TAG, "*** Engine::ReceiveEvent - AARVersion and OBBVersion mismatch!");
+			}
+
+			if (pendingPakList != null) {
+				// mount any pending pak files with higher priority
+				for (String pakFile : pendingPakList) {
+					nativeMountPak(pakFile, 1, null);
+				}
+				pendingPakList = null;
+			}
+		}
+		else if (event == EVENTTYPE_ENGINELOOP_INIT_COMPLETE) {
+			Log.d(TAG, "*** Engine::ReceiveEvent - event = " + EVENTTYPE_ENGINELOOP_INIT_COMPLETE + ", param1 =" + param1);
+		}
+		else if (event == EVENTTYPE_POST_ENGINE_INIT) {
+			Log.d(TAG, "*** Engine::ReceiveEvent - event = " + EVENTTYPE_POST_ENGINE_INIT + ", param1 =" + param1);
+		}
+		else if (event == EVENTTYPE_POST_LOAD_MAP) {
+			Log.d(TAG, "*** Engine::ReceiveEvent - event = " + EVENTTYPE_POST_LOAD_MAP + ", param1 =" + param1);
+		}
+		else if (event == EVENTTYPE_PRE_LOAD_MAP) {
+			Log.d(TAG, "*** Engine::ReceiveEvent - event = " + EVENTTYPE_PRE_LOAD_MAP + ", param1 =" + param1);
+		}
+		else if (event == EVENTTYPE_ACTIVITY_HAS_CHANGED) {
+			Log.d(TAG, "*** Engine::ReceiveEvent - event = " + EVENTTYPE_ACTIVITY_HAS_CHANGED + ", param1 =" + param1);
+		}
+		else  if (event != EVENTTYPE_FRAME_BEGIN && event != EVENTTYPE_FRAME_END)
+		{
+			Log.d(TAG, "*** Engine::ReceiveEvent - event = " + event + ", param1 =" + param1);
+		}
+
+		maskEngineEvents |= 1<<event;
+
+		final IEventCallback callback = mCallbackRef.get();
+		if (callback != null) {
+			callback.eventCallback(event, param1, param2, param3, param4);
+		}
+	}
+
+	public void sendEvent(int event, String param1, int param2, int param3, float param4)
+	{
+		nativeEventReceived(event, param1, param2, param3, param4);
+	}
+
+	public void sendData(int dataId, String key, String value)
+	{
+		nativeDataReceivedString(dataId, key, value);
+	}
+
+	public void sendData(int dataId, String key, String[] value)
+	{
+		nativeDataReceivedStringArray(dataId, key, value);
+	}
+
+	public void sendData(int dataId, String key, boolean value)
+	{
+		nativeDataReceivedBoolean(dataId, key, value);
+	}
+
+	public void sendData(int dataId, String key, boolean[] value)
+	{
+		nativeDataReceivedBooleanArray(dataId, key, value);
+	}
+
+	public void sendData(int dataId, String key, byte value)
+	{
+		nativeDataReceivedByte(dataId, key, value);
+	}
+
+	public void sendData(int dataId, String key, byte[] value)
+	{
+		nativeDataReceivedByteArray(dataId, key, value);
+	}
+
+	public void sendData(int dataId, String key, int value)
+	{
+		nativeDataReceivedInt(dataId, key, value);
+	}
+
+	public void sendData(int dataId, String key, int[] value)
+	{
+		nativeDataReceivedIntArray(dataId, key, value);
+	}
+
+	public void sendData(int dataId, String key, Short value)
+	{
+		nativeDataReceivedShort(dataId, key, value);
+	}
+
+	public void sendData(int dataId, String key, Short[] value)
+	{
+		nativeDataReceivedShortArray(dataId, key, value);
+	}
+
+	public void sendData(int dataId, String key, Long value)
+	{
+		nativeDataReceivedLong(dataId, key, value);
+	}
+
+	public void sendData(int dataId, String key, Long[] value)
+	{
+		nativeDataReceivedLongArray(dataId, key, value);
+	}
+
+	public void sendData(int dataId, String key, float value)
+	{
+		nativeDataReceivedFloat(dataId, key, value);
+	}
+
+	public void sendData(int dataId, String key, float[] value)
+	{
+		nativeDataReceivedFloatArray(dataId, key, value);
+	}
+
+	public void sendData(int dataId, String key, Double value)
+	{
+		nativeDataReceivedDouble(dataId, key, value);
+	}
+
+	public void sendData(int dataId, String key, Double[] value)
+	{
+		nativeDataReceivedDoubleArray(dataId, key, value);
+	}
+
+	public boolean mountPak(String pakFile, int order, String mountPoint)
+	{
+		if (pakFile != null) {
+			return nativeMountPak(pakFile, order, mountPoint == null ? "" : mountPoint);
+		}
+		return false;
+	}
+
+	public boolean checkPak(String pakFile)
+	{
+		if (pakFile != null) {
+			return nativeCheckPak(pakFile);
+		}
+		return false;
+	}
+
+	public static native void nativeEventReceived(int event, String param1, int param2, int param3, float param4);
+	public static native void nativeDataReceivedString(int dataId, String key, String value);
+	public static native void nativeDataReceivedStringArray(int dataId, String key, String[] value);
+	public static native void nativeDataReceivedBoolean(int dataId, String key, boolean value);
+	public static native void nativeDataReceivedBooleanArray(int dataId, String key, boolean[] value);
+	public static native void nativeDataReceivedByte(int dataId, String key, byte value);
+	public static native void nativeDataReceivedByteArray(int dataId, String key, byte[] value);
+	public static native void nativeDataReceivedInt(int dataId, String key, int value);
+	public static native void nativeDataReceivedIntArray(int dataId, String key, int[] value);
+	public static native void nativeDataReceivedShort(int dataId, String key, Short value);
+	public static native void nativeDataReceivedShortArray(int dataId, String key, Short[] value);
+	public static native void nativeDataReceivedLong(int dataId, String key, Long value);
+	public static native void nativeDataReceivedLongArray(int dataId, String key, Long[] value);
+	public static native void nativeDataReceivedFloat(int dataId, String key, float value);
+	public static native void nativeDataReceivedFloatArray(int dataId, String key, float[] value);
+	public static native void nativeDataReceivedDouble(int dataId, String key, Double value);
+	public static native void nativeDataReceivedDoubleArray(int dataId, String key, Double[] value);
+	public static native boolean nativeMountPak(String pakFile, int order, String mountPoint);
+	public static native boolean nativeCheckPak(String pakFile);
+}
\ No newline at end of file
diff --git a/Engine/Plugins/Runtime/AndroidSingleInstanceService/Source/AndroidSingleInstanceService/Android/src/main/java/com/epicgames/makeaar/EngineFactory.java b/Engine/Plugins/Runtime/AndroidSingleInstanceService/Source/AndroidSingleInstanceService/Android/src/main/java/com/epicgames/makeaar/EngineFactory.java
new file mode 100644
--- /dev/null
+++ b/Engine/Plugins/Runtime/AndroidSingleInstanceService/Source/AndroidSingleInstanceService/Android/src/main/java/com/epicgames/makeaar/EngineFactory.java
@@ -0,0 +1,98 @@
+package com.epicgames.makeaar;
+
+import android.app.Activity;
+import android.app.Application;
+import android.util.Log;
+
+//import com.epicgames.makeaar.Engine;
+
+import android.os.HandlerThread;
+import android.os.Looper;
+
+/**
+ * Shared singleton EngineFactory thread for the system.  This is a thread for
+ * EngineFactory connectivity operations such as AsyncChannel connections to system services.
+ * Various connectivity manager objects can use this singleton as a common
+ * resource for their handlers instead of creating separate threads of their own.
+ * @hide
+ */
+public class EngineFactory extends HandlerThread {
+	static Engine engine;
+	//Activity engineActivity = null;
+
+	// A class implementing the lazy holder idiom: the unique static instance
+	// of ConnectivityThread is instantiated in a thread-safe way (guaranteed by
+	// the language specs) the first time that Singleton is referenced in get()
+	// or getInstanceLooper().
+	private static class Singleton {
+		private static final EngineFactory INSTANCE = createInstance();
+	}
+
+	private EngineFactory() {
+		super("EngineFactory");
+	}
+
+	private static EngineFactory createInstance() {
+		EngineFactory t = new EngineFactory();
+		t.start();
+		return t;
+	}
+
+	public static EngineFactory get() {
+		return Singleton.INSTANCE;
+	}
+	public static Engine getEngine() {
+		return Singleton.INSTANCE.engine;
+	}
+
+	public static Looper getInstanceLooper() {
+		return Singleton.INSTANCE.getLooper();
+	}
+
+	public static void Reset()
+	{
+		if (Singleton.INSTANCE.engine != null)
+		{
+			Log.d("UE", "EngineFactory::Reset() - releasing current engine resources, engine = " + Singleton.INSTANCE.engine);
+			//engine.onDestroy(543210);
+			Singleton.INSTANCE.engine.onDestroy(543210);
+		}
+		
+//		engineActivity = null;
+		Singleton.INSTANCE.engine = null;
+	}
+
+	public static Engine getInstance(Activity activity, String OBBFilename, String projectModule, boolean enablePropagateAlpha) {
+		
+		if (Singleton.INSTANCE.engine != null)
+		{
+            Log.d("UE", "EngineFactory::GetInstance() - USING existing Engine instance, proc = " + Application.getProcessName() + ", engine = " + Singleton.INSTANCE.engine);
+			return Singleton.INSTANCE.engine;
+		}
+//		engineActivity = activity;
+		Singleton.INSTANCE.engine = new Engine(activity, OBBFilename, projectModule, enablePropagateAlpha);
+		
+        Log.d("UE", "EngineFactory::GetInstance(with OBBFilename) - CREATED new Engine instance, proc = " + Application.getProcessName() + ", engine = " + Singleton.INSTANCE.engine);
+        return Singleton.INSTANCE.engine;
+	}
+
+	public static Engine getInstance(Activity activity) {
+		return getInstance(activity, "", "", false);
+	}
+
+	public static Engine getInstance(Activity activity, String projectModule) {
+		return getInstance(activity, "", projectModule, false);
+	}
+
+	public static Engine getInstance(Activity activity, String OBBFilename, String projectModule) {
+		return getInstance(activity, OBBFilename, projectModule, false);
+	}
+
+	public static Engine getInstance(Activity activity, String projectModule, boolean enablePropagateAlpha) {
+		return getInstance(activity, "", projectModule, enablePropagateAlpha);
+	}
+
+	public static Engine getInstance(String OBBFilename, String projectModule, boolean enablePropagateAlpha) {
+		return getInstance(null, OBBFilename, projectModule, enablePropagateAlpha);
+	}
+}
diff --git a/Engine/Plugins/Runtime/AndroidSingleInstanceService/Source/AndroidSingleInstanceService/Android/src/main/java/com/epicgames/makeaar/GameActivityForMakeAAR.java b/Engine/Plugins/Runtime/AndroidSingleInstanceService/Source/AndroidSingleInstanceService/Android/src/main/java/com/epicgames/makeaar/GameActivityForMakeAAR.java
new file mode 100644
--- /dev/null
+++ b/Engine/Plugins/Runtime/AndroidSingleInstanceService/Source/AndroidSingleInstanceService/Android/src/main/java/com/epicgames/makeaar/GameActivityForMakeAAR.java
@@ -0,0 +1,562 @@
+package com.epicgames.makeaar;
+
+import android.app.Activity;
+import android.content.Context;
+import android.content.ContextWrapper;
+
+import java.io.File;
+import java.util.Map;
+import java.util.HashMap;
+
+import android.content.pm.ApplicationInfo;
+import android.content.pm.FeatureInfo;
+import android.os.Bundle;
+import android.os.Looper;
+
+import android.os.Handler;
+
+import android.content.ClipboardManager;
+import android.content.res.Configuration;
+import android.content.pm.PackageManager;
+import android.content.pm.PackageInfo;
+
+import android.media.AudioManager;
+
+import android.view.View;
+import android.view.Surface;
+
+import com.epicgames.unreal.GameActivity;
+import com.epicgames.unreal.Logger;
+import com.epicgames.unreal.SimpleContextWrapper;
+
+public class GameActivityForMakeAAR extends com.epicgames.unreal.GameActivity
+{
+	//static GameActivityForMakeAAR _gameActivityInstance = null;
+
+	//========== MAKEAAR specifics ====================
+
+	//public SimpleContextWrapper activityContext = null;
+	
+
+	String _OBBFilename = null;
+	boolean _enablePropagateAlpha = false;
+
+	public static com.epicgames.unreal.Logger Log = new Logger("UE_ASIS", "MakeAARBaseForGameActivity");
+	
+//	String commandLine = "";
+
+	GameActivityForMakeAAR()
+	{
+		isStandalone = true;
+		Log.debug("GameActivityForMakeAAR constructor called with _gameActivityInstance=" + _gameActivityInstance + ", activityContext=" + activityContext + ", this=" + this + ", isStandalone=" + isStandalone);
+		_gameActivityInstance = this;
+		//activityContext = new SimpleContextWrapper(this);
+	}
+
+	public static GameActivityForMakeAAR Instance()
+	{
+		if (_gameActivityInstance != null)
+		{
+			Log.error("GameActivityForMakeAAR.Instance() called but _gameActivityInstance already exists!!! in current thread= " + Thread.currentThread() + ", _gameActivityInstance=" + _gameActivityInstance);
+			
+		}
+		else
+		{
+			GameActivityForMakeAAR newInstance = new GameActivityForMakeAAR();
+			Log.debug("GameActivityForMakeAAR::Instance() constructor called with prev _gameActivityInstance=" + _gameActivityInstance + ", newInstance.activityContext=" + newInstance.activityContext + ", newInstance=" + newInstance + ", newInstance.isStandalone=" + newInstance.isStandalone);
+			_gameActivityInstance = newInstance;
+			
+		}
+
+		assert( _gameActivityInstance instanceof GameActivityForMakeAAR);
+
+		return (GameActivityForMakeAAR)_gameActivityInstance;
+	}
+
+	public static GameActivityForMakeAAR Create()
+	{
+		return Instance();
+	}
+
+	
+	void setCommandline(String inCommandline)
+	{
+		nativeSetCommandline(inCommandline);
+	}
+
+	public static Boolean isValidGameActivity()
+	{
+		return GameActivity.Get() != null;
+	}
+
+	public static GameActivityForMakeAAR Get()
+	{
+		return (GameActivityForMakeAAR)_gameActivityInstance;
+	}
+
+	@Override
+	protected void RestartApplication(String RestartExtra)
+	{
+		super.RestartApplication(RestartExtra);
+		//Context context = activityContext.getApplicationContext();
+		Log.warn("GameActivityForMakeAAR.RestartApplication() in current thread= " + Thread.currentThread() + ", _gameActivityInstance=" + _gameActivityInstance);
+	}
+
+
+	// ========	override defined in:	AndroidJNI.cpp
+	public native void nativeSetGlobalActivity(boolean bUseExternalFilesDir, boolean bPublicLogFiles, String internalFilePath, String externalFilePath, boolean bOBBInAPK, String APKPath);
+
+	// ========	override defined in:	AndroidWindow.cpp
+	public native static void nativeSetSurfaceOverride(Surface surface, int x, int y);
+
+	// ========	override defined in:	AndroidPlatformFile.cpp
+	public native static String nativeGetObbComment();
+
+
+	// ========	MakeAAR specific defined in:	LaunchAndroid.cpp
+	public native static void nativeMain(String projectModule);
+	public native static void nativeSetCommandline(String commandLine);
+	public native static void nativeAppCommand(int cmd);
+	public native static int nativeInputTouch(int device, int action, int pointerId, int x, int y);
+
+
+	// ========            MAKEAAR specific non-native functions			===========
+	
+	public boolean setActivity( Activity inActivity, String OBBFilename, boolean enablePropagateAlpha) {
+		return setActivity(new SimpleContextWrapper(inActivity), OBBFilename, enablePropagateAlpha);
+	}
+
+	public boolean setActivity(SimpleContextWrapper inActivity, String OBBFilename, boolean enablePropagateAlpha) {
+
+
+		Log.verbose( "**LIFECYCLE** GameActivity::setActivity(with OBBFilename) inActivity=" + inActivity + ", _activity=" + activityContext + ", obb=" + OBBFilename);
+
+		if (activityContext != null)
+		{
+			// we need to shut down the old activity
+			Log.debug( "**IMPORTANT** GameActivity::setActivity(with OBBFilename) change inActivity=" + inActivity + ", _activity=" + activityContext);
+			Log.debug( "**TODO** cleanup old activity");
+
+		}
+
+		activityContext = inActivity;
+		
+		if (activityContext == null)
+		{
+			Log.error( "GameActivity::setActivity(with OBBFilename) _activity is null" );
+			return false;
+		}
+
+		clipboardManager = (ClipboardManager) activityContext.getSystemService(CLIPBOARD_SERVICE);
+
+		// Grab a reference to the asset manager
+		AssetManagerReference = activityContext.getAssets();
+
+		InternalFilesDir = activityContext.getFilesDir().getAbsolutePath() + "/";
+		File _externalFilesDir = activityContext.getExternalFilesDir(null);
+		ExternalFilesDir = (_externalFilesDir != null ? _externalFilesDir.getAbsolutePath() : "") + "/";
+		
+		_extrasBundle = activityContext.mAppActivityContext != null && activityContext.mAppActivityContext instanceof Activity ? ((Activity)activityContext.mAppActivityContext).getIntent().getExtras() : null;
+		appPackageName = activityContext.getPackageName();
+		String ProjectName = appPackageName;
+		ProjectName = ProjectName.substring(ProjectName.lastIndexOf('.') + 1);
+		Log.verbose("from SetActivity(with OBBFilename), appPackageName: " + appPackageName + ", ProjectName = " + ProjectName);
+
+		String AppType = "";
+		try
+		{
+			ApplicationInfo ai = activityContext.getPackageManager().getApplicationInfo(appPackageName, PackageManager.GET_META_DATA);
+			Bundle bundle = ai.metaData;
+			_bundle = bundle;
+
+			if ((ai.flags & ApplicationInfo.FLAG_DEBUGGABLE) == 0) {
+				IsForDistribution = true;
+			}
+		}
+		catch (PackageManager.NameNotFoundException e)
+		{
+			Log.debug( "Failed to load meta-data: NameNotFound: " + e.getMessage());
+		}
+		catch (NullPointerException e)
+		{
+			Log.debug( "Failed to load meta-data: NullPointer: " + e.getMessage());
+		}
+
+		// Look for Vulkan support if Nougat or later
+		if (ANDROID_BUILD_VERSION >= 24)
+		{
+			FeatureInfo[] features = activityContext.getPackageManager().getSystemAvailableFeatures();
+			for (FeatureInfo feature : features) {
+				if (feature.name != null)
+				{
+					if (feature.name.equals("android.hardware.vulkan.level"))
+					{
+						// since we may not be compiled against android-24 or higher, use .toString to get the version field
+						String dump = feature.toString();
+						int index = dump.indexOf("v=");
+						if (index >= 0)
+						{
+							dump = dump.substring(index+2);
+							index = dump.indexOf(" ");
+							if (index >= 0)
+							{
+								VulkanLevel = Integer.parseInt(dump.substring(0, index));
+								Log.debug("Vulkan level: " + VulkanLevel);
+							}
+						}
+					}
+					else
+					if (feature.name.equals("android.hardware.vulkan.version"))
+					{
+						// since we may not be compiled against android-24 or higher, use .toString to get the version field
+						String dump = feature.toString();
+						int index = dump.indexOf("v=");
+						if (index >= 0)
+						{
+							dump = dump.substring(index+2);
+							index = dump.indexOf(" ");
+							if (index >= 0)
+							{
+								VulkanVersion = Integer.parseInt(dump.substring(0, index));
+								int VersionMajor = (VulkanVersion >> 22) & 0x03ff;
+								int VersionMinor = (VulkanVersion >> 12) & 0x03ff;
+								int VersionPatch = VulkanVersion & 0x0fff;
+								VulkanVersionString = VersionMajor + "." + VersionMinor + "." + VersionPatch;
+								Log.debug("SetActivity: Vulkan version: " + VersionMajor + "." + VersionMinor + "." + VersionPatch);
+							}
+						}
+					}
+				}
+			}
+		}
+
+		boolean bDebuggerAttached = android.os.Debug.isDebuggerConnected();
+		nativeSetAndroidStartupState(bDebuggerAttached);
+
+		_OBBFilename = OBBFilename;
+		_enablePropagateAlpha = enablePropagateAlpha;
+
+		GameActivity.setOBBInAPK(OBBFilename.equals(""));
+		Log.debug("APK path: " + activityContext.getPackageResourcePath());
+		Log.debug("OBB in APK: " + (PackageDataInsideApkValue==1));
+		boolean UseExternalFilesDir = false;
+		boolean PublicLogFiles = false;
+		int DepthBufferPreference = 32;
+		int PropagateAlpha = enablePropagateAlpha ? 1 : 0;
+
+		Log.debug("SetActivity(with OBBFilename) ExternalFilesDir path: " + ExternalFilesDir);
+		Log.debug("activityContext.getExternalFilesDir(null).getPath(): " + activityContext.getExternalFilesDir(null).getPath());
+
+		nativeSetGlobalActivity(UseExternalFilesDir, PublicLogFiles, activityContext.getFilesDir().getPath(), activityContext.getExternalFilesDir(null).getPath(), PackageDataInsideApkValue==1, activityContext.getPackageResourcePath());
+		if (PackageDataInsideApkValue == 0)
+		{
+			Log.debug("OBB override: " + OBBFilename);
+			nativeSetObbFilePaths(OBBFilename, "", "", "");
+		}
+
+		// tell the engine if this is a portrait app and if it will propagate alpha
+		boolean bPortrait = activityContext.getResources().getConfiguration().orientation == Configuration.ORIENTATION_PORTRAIT;
+		nativeSetWindowInfo(bPortrait, DepthBufferPreference, PropagateAlpha);
+
+		// get the full language code, like en-US
+		// note: this may need to be Locale.getDefault().getLanguage()
+		String Language = java.util.Locale.getDefault().toString();
+
+		// Retrieve version code, name, and target SDK
+		try 
+		{
+			PackageInfo packageInfo = activityContext.getPackageManager().getPackageInfo(activityContext.getPackageName(), 0);
+			VersionCode = packageInfo.versionCode;
+			VersionName = packageInfo.versionName;
+			targetSdkVersion = packageInfo.applicationInfo.targetSdkVersion;
+		}
+		catch (Exception e)
+		{
+			Log.debug("Error accessing packageInfo: " + e.getMessage());
+		}
+
+		String productName = getProductName();
+		Log.debug("Android version is " + android.os.Build.VERSION.RELEASE);
+		Log.debug("Android manufacturer is " + android.os.Build.MANUFACTURER);
+		Log.debug("Android model is " + android.os.Build.MODEL);
+		Log.debug("Android build number is " + android.os.Build.DISPLAY);
+		Log.debug("OS language is set to " + Language);
+		Log.debug("Product name is " + (productName.isEmpty() ? "[not set]" : productName));
+		Log.debug("Debugger attached is " + bDebuggerAttached);
+		Log.debug("Android targetSdkVersion version is " + targetSdkVersion);
+
+		nativeSetAndroidVersionInformation(android.os.Build.VERSION.RELEASE, targetSdkVersion, android.os.Build.MANUFACTURER, android.os.Build.MODEL, android.os.Build.DISPLAY, Language, productName);
+
+
+		try
+		{
+			VersionCode = activityContext.getPackageManager().getPackageInfo(appPackageName, 0).versionCode;
+			int PatchVersion = 0;
+			nativeSetObbInfo(ProjectName, activityContext.getApplicationContext().getPackageName(), VersionCode, PatchVersion, AppType);
+		}
+		catch (Exception e)
+		{
+			// if the above failed, then, we can't use obbs
+			Log.debug("==================================> PackageInfo failure getting .obb info: " + e.getMessage());
+		}
+
+		// tell Android that we want volume controls to change the media volume, aka music
+		// enable the physical volume controls to the game
+		if (inActivity != null && activityContext.mAppActivityContext != null && activityContext.mAppActivityContext instanceof Activity) {
+			((Activity)activityContext.mAppActivityContext).setVolumeControlStream(AudioManager.STREAM_MUSIC);
+		}
+
+		// store the screen resolution (adjusted for immersive mode)
+		Map<String, String> variables = new HashMap<String, String>();
+		android.view.Display display = activityContext.getWindowManager().getDefaultDisplay();
+		android.graphics.Point displaySize = new android.graphics.Point();
+
+		// check if immersive mode on window
+//		View decorView = activityContext.getWindow().getDecorView();
+//		if ((decorView.getSysteUiVisibility() & View.System_UI_FLAG_IMMERSIVE_STICKY) != 0)
+		{
+			display.getRealSize(displaySize);
+		}
+//		else
+//		{
+//			display.getSize(displaySize);
+//		}
+		variables.put("screenWidth", Integer.toString(displaySize.x));
+		variables.put("screenHeight", Integer.toString(displaySize.y));
+
+		// send the variables to the engine
+		String[] keyValues = new String[variables.size() * 2];
+		int keyValueIndex = 0;
+		for (Map.Entry<String, String> entry : variables.entrySet())
+		{
+			keyValues[keyValueIndex++] = entry.getKey();
+			keyValues[keyValueIndex++] = entry.getValue();
+		}
+		GameActivity.Get().nativeSetConfigRulesVariables(keyValues);
+
+		return true;
+	}
+
+	public boolean setActivity(SimpleContextWrapper inActivity, GameActivitySetupInfo gameActivitySetupInfo)
+	{
+		if (activityContext != null)
+		{
+			// we need to shut down the old activity
+			Log.error( "**IMPORTANT** GameActivity::setActivity(with gameActivitySetupInfo) change inActivity=" + inActivity + ", activityContext=" + activityContext );
+			Log.error( "**TODO** cleanup old activity");
+		}
+
+		activityContext = inActivity;
+		Log.verbose( "**LIFECYCLE** GameActivity::setActivity(with gameActivitySetupInfo) inActivity=" + inActivity + ", activityContext=" + activityContext);
+
+
+		if (activityContext == null)
+		{
+			Log.error( "GameActivity::setActivity(with gameActivitySetupInfo) activityContext is null" );
+		}
+		
+
+
+		
+		clipboardManager = gameActivitySetupInfo.clipboardManager;
+
+		// Grab a reference to the asset manager
+		AssetManagerReference = gameActivitySetupInfo.AssetManagerReference;
+
+		InternalFilesDir = gameActivitySetupInfo.InternalFilesDir;
+		ExternalFilesDir = gameActivitySetupInfo.ExternalFilesDir;
+
+		_extrasBundle = gameActivitySetupInfo._extrasBundle;
+
+		if (this._extrasBundle == null)
+		{
+			Log.error("=====> PackageInfo _extrasBundle so no renderSurface is passed!" );
+		}
+		else if (gameActivitySetupInfo.renderSurface == null)
+		{
+			gameActivitySetupInfo.renderSurface = this._extrasBundle == null ? null : this._extrasBundle.<Surface>getParcelable("renderSurface");
+			Log.warn("====> GameActivitySetupInfo(GET) is storing renderSurface = " + gameActivitySetupInfo.renderSurface);
+		}
+		else
+		{
+			Log.debug("====> GameActivitySetupInfo(GET) is has renderSurface = " + gameActivitySetupInfo.renderSurface);
+
+		}
+
+		appPackageName = gameActivitySetupInfo.appPackageName;
+
+		String ProjectName = appPackageName;
+		ProjectName = gameActivitySetupInfo.ProjectName.isEmpty() ? ProjectName.substring(ProjectName.lastIndexOf('.') + 1) : gameActivitySetupInfo.ProjectName;
+		String AppType = "";
+		Log.verbose("from gameActivitySetupInfo, appPackageName: " + appPackageName + ", ProjectName = " + ProjectName);
+
+		_bundle = gameActivitySetupInfo._bundle;
+		IsForDistribution = gameActivitySetupInfo.IsForDistribution;
+
+
+		// Look for Vulkan support if Nougat or later
+		gameActivitySetupInfo.GetVulkanInfoFromPackage(this);
+
+		boolean bDebuggerAttached = android.os.Debug.isDebuggerConnected();
+		nativeSetAndroidStartupState(bDebuggerAttached);
+
+		_OBBFilename = gameActivitySetupInfo._OBBFilename;
+		int PropagateAlpha = gameActivitySetupInfo.PropagateAlpha;
+
+		_enablePropagateAlpha = gameActivitySetupInfo.PropagateAlpha == 0 ? false : true;
+
+		GameActivity.setOBBInAPK(_OBBFilename.equals(""));
+
+		Log.debug("APK path: " + activityContext.getPackageResourcePath());
+		Log.debug("OBB in APK: " + (PackageDataInsideApkValue==1));
+		boolean UseExternalFilesDir = false;
+		boolean PublicLogFiles = false;
+		int DepthBufferPreference = 32;
+
+		Log.debug("SetActivity(with gameActivitySetupInfo) ExternalFilesDir path: " + ExternalFilesDir);
+		Log.debug("activityContext.getExternalFilesDir(null).getPath(): " + activityContext.getExternalFilesDir(null).getPath());
+
+		nativeSetGlobalActivity(UseExternalFilesDir, PublicLogFiles, activityContext.getFilesDir().getPath(), activityContext.getExternalFilesDir(null).getPath(), PackageDataInsideApkValue==1, activityContext.getPackageResourcePath());
+		if (PackageDataInsideApkValue == 0)
+		{
+			Log.debug("OBB override: " + _OBBFilename);
+			nativeSetObbFilePaths(_OBBFilename, "", "", "");
+		}
+
+		// tell the engine if this is a portrait app and if it will propagate alpha
+		boolean bPortrait = activityContext.getResources().getConfiguration().orientation == Configuration.ORIENTATION_PORTRAIT;
+		nativeSetWindowInfo(bPortrait, DepthBufferPreference, PropagateAlpha);
+
+		// get the full language code, like en-US
+		// note: this may need to be Locale.getDefault().getLanguage()
+		String Language = java.util.Locale.getDefault().toString();
+
+		// Retrieve version code, name, and target SDK
+		try
+		{
+			PackageInfo packageInfo = activityContext.getPackageManager().getPackageInfo(activityContext.getPackageName(), 0);
+			VersionCode = packageInfo.versionCode;
+			VersionName = packageInfo.versionName;
+			targetSdkVersion = packageInfo.applicationInfo.targetSdkVersion;
+		}
+		catch (Exception e)
+		{
+			Log.debug("Error accessing packageInfo: " + e.getMessage());
+		}
+
+		String productName = getProductName();
+		Log.debug( "Android version is " + android.os.Build.VERSION.RELEASE );
+		Log.debug( "Android manufacturer is " + android.os.Build.MANUFACTURER );
+		Log.debug( "Android model is " + android.os.Build.MODEL );
+		Log.debug( "Android build number is " + android.os.Build.DISPLAY );
+		Log.debug( "OS language is set to " + Language);
+		Log.debug( "Product name is " + (productName.isEmpty() ? "[not set]" : productName));
+		Log.debug( "Debugger attached is " + bDebuggerAttached );
+		Log.debug( "Android targetSdkVersion version is " + targetSdkVersion );
+		
+		nativeSetAndroidVersionInformation(android.os.Build.VERSION.RELEASE, targetSdkVersion, android.os.Build.MANUFACTURER, android.os.Build.MODEL, android.os.Build.DISPLAY, Language, productName);
+
+
+		try
+		{
+			Log.debug( "do nativeSetObbInfo ProjectName=" + ProjectName + ", PackageName=" + appPackageName );
+
+			//VersionCode = activityContext.getPackageManager().getPackageInfo(appPackageName, 0).versionCode;
+			VersionCode = gameActivitySetupInfo.VersionCode;
+			int PatchVersion = 0;
+			//nativeSetObbInfo(ProjectName, activityContext.getApplicationContext().getPackageName(), VersionCode, PatchVersion, AppType);
+			nativeSetObbInfo(ProjectName, gameActivitySetupInfo.appPackageName, VersionCode, PatchVersion, AppType);
+		}
+		catch (Exception e)
+		{
+			// if the above failed, then, we can't use obbs
+			Log.debug("==================================> PackageInfo failure getting .obb info: " + e.getMessage());
+		}
+		this._extrasBundle = activityContext.mAppActivityContext != null && activityContext.mAppActivityContext instanceof Activity ? ((Activity)activityContext.mAppActivityContext).getIntent().getExtras() : null;
+
+		// tell Android that we want volume controls to change the media volume, aka music
+		// enable the physical volume controls to the game
+		if (inActivity != null && activityContext.mAppActivityContext != null && activityContext.mAppActivityContext instanceof Activity) {
+			 ((Activity)activityContext.mAppActivityContext).setVolumeControlStream(AudioManager.STREAM_MUSIC);
+		}
+
+		String[] keyValues = gameActivitySetupInfo.keyValues;
+
+		if (keyValues.length == 0) {
+			// store the screen resolution (adjusted for immersive mode)
+			Map<String, String> variables = new HashMap<String, String>();
+			android.view.Display display = activityContext.getWindowManager().getDefaultDisplay();
+			android.graphics.Point displaySize = new android.graphics.Point();
+
+			// check if immersive mode on window
+			View decorView = activityContext.getWindow().getDecorView();
+	//		if ((decorView.getSysteUiVisibility() & View.System_UI_FLAG_IMMERSIVE_STICKY) != 0)
+			{
+				display.getRealSize(displaySize);
+			}
+	//		else
+	//		{
+	//			display.getSize(displaySize);
+	//		}
+			variables.put("screenWidth", Integer.toString(displaySize.x));
+			variables.put("screenHeight", Integer.toString(displaySize.y));
+
+			// send the variables to the engine
+			keyValues = new String[variables.size() * 2];
+			int keyValueIndex = 0;
+			for (Map.Entry<String, String> entry : variables.entrySet())
+			{
+				keyValues[keyValueIndex++] = entry.getKey();
+				keyValues[keyValueIndex++] = entry.getValue();
+			}
+		}
+
+		nativeSetConfigRulesVariables(keyValues);
+
+		return true;
+	}
+
+//	public void GameActivityForMakeAAR(Activity inActivity, String OBBFilename, boolean enablePropagateAlpha)
+//	{
+//		gameActivity = this;
+//
+//		//SimpleContextWrapper inActivity = SimpleContextWrapper.getSimpleActivityOfContext(inContext);
+//		//Activity inActivity = SimpleContextWrapper.getActivityOfContext(inContext);
+//		
+//		setActivity(inActivity, OBBFilename, enablePropagateAlpha);
+//	}
+//
+//	public void GameActivityForMakeAAR(Context context, GameActivitySetupInfo gameActivitySetupInfo)
+//	{
+//		gameActivity = this;
+//		activityContext = new SimpleContextWrapper(context);
+//		setActivity(activityContext, gameActivitySetupInfo);
+//	}
+//
+//	public void GameActivityForMakeAAR(GameActivitySetupInfo gameActivitySetupInfo)
+//	{
+//		gameActivity = this;
+//		initFromData(gameActivitySetupInfo);
+//	}
+
+	void initFromData(GameActivitySetupInfo gameActivitySetupInfo)
+	{
+
+		AssetManagerReference = gameActivitySetupInfo.AssetManagerReference;
+		clipboardManager = gameActivitySetupInfo.clipboardManager;
+
+		//Activity _activity = null;
+		_OBBFilename = gameActivitySetupInfo._OBBFilename;
+		_enablePropagateAlpha = gameActivitySetupInfo.PropagateAlpha == 1;
+		_bundle = gameActivitySetupInfo._bundle;
+		_extrasBundle = gameActivitySetupInfo._extrasBundle;
+		PackageDataInsideApkValue = gameActivitySetupInfo.packageDataInsideApkValue ? 1 : 0;
+		InternalFilesDir = gameActivitySetupInfo.InternalFilesDir;
+		ExternalFilesDir = gameActivitySetupInfo.ExternalFilesDir;
+		appPackageName = gameActivitySetupInfo.appPackageName;
+		VersionCode = gameActivitySetupInfo.VersionCode;
+		//---- Discovered Vulkan Version and Level from getSystemAvailableFeatures() ---//
+		gameActivitySetupInfo.GetVulkanInfoFromPackage(this);
+
+	}
+}
\ No newline at end of file
