Android SingleInstance Service

diff --git a/Engine/Plugins/Runtime/AndroidSingleInstanceService/Source/AndroidSingleInstanceService/Android/src/main/java/com/epicgames/makeaar/GameActivitySetupInfo.java b/Engine/Plugins/Runtime/AndroidSingleInstanceService/Source/AndroidSingleInstanceService/Android/src/main/java/com/epicgames/makeaar/GameActivitySetupInfo.java
new file mode 100644
--- /dev/null
+++ b/Engine/Plugins/Runtime/AndroidSingleInstanceService/Source/AndroidSingleInstanceService/Android/src/main/java/com/epicgames/makeaar/GameActivitySetupInfo.java
@@ -0,0 +1,402 @@
+package com.epicgames.makeaar;
+
+import android.app.Activity;
+import android.content.ClipboardManager;
+import android.content.Context;
+import android.content.pm.ApplicationInfo;
+import android.content.pm.FeatureInfo;
+import android.content.pm.PackageInfo;
+import android.content.res.AssetManager;
+import android.content.res.Configuration;
+import android.os.Bundle;
+import android.view.Surface;
+import android.content.pm.PackageManager;
+import android.os.Parcel;
+import android.os.Parcelable;
+import android.util.Log;
+
+import java.io.File;
+import java.util.Map;
+import java.util.HashMap;
+
+//import com.epicgames.unreal.Logger;
+
+import static android.content.Context.CLIPBOARD_SERVICE;
+import static android.content.Context.WINDOW_SERVICE;
+
+import com.epicgames.unreal.SimpleContextWrapper;
+
+//import androidx.annotation.RequiresApi;
+
+public class GameActivitySetupInfo implements Parcelable {
+
+	//public static Logger Log = new Logger("UES", "GameActivitySetupInfo");
+	final static String TAG = "UE_ASIS_GameActivitySetupInfo";
+
+	public static final int MSG_ATTACH_EXTERNAL_SURFACE = 2;
+		
+	// Class Attributes
+	public String InternalFilesDir;
+	public String ExternalFilesDir;
+	public String packageResourcePath;
+	public boolean OrientationPortrait;
+	public int DepthBufferPreference = 32;
+	public int PropagateAlpha;
+	public boolean UseExternalFilesDir = false;
+	public boolean PublicLogFiles = false;
+
+	public String _OBBFilename;
+	public boolean packageDataInsideApkValue;
+	public int[] ViewPosition = new int[2];
+	public int[] ViewSize = new int[2];
+
+
+	public int VersionCode;
+	public int targetSdkVersion;
+
+	public ClipboardManager clipboardManager;
+	public AssetManager AssetManagerReference;
+	public PackageManager packageManager;
+	
+	public Bundle _extrasBundle;
+	public Bundle _bundle;
+	public String appPackageName;
+	public String ProjectName;
+	public boolean IsForDistribution;
+	public String[] keyValues;
+	public Surface renderSurface;
+
+	/**
+	 Get the SDK level of the OS we are running in.
+	 We do this instead of accessing the SDK_INT
+	 with JNI from C++ as the new ART runtime seems to have
+	 problems dynamically finding/loading static inner classes.
+	 */
+	public static final int ANDROID_BUILD_VERSION = android.os.Build.VERSION.SDK_INT;
+
+
+
+	public void GetVulkanInfoFromPackage( GameActivityForMakeAAR ref )
+	{
+		// Look for Vulkan support if Nougat or later
+		if (ANDROID_BUILD_VERSION >= 24)
+		{
+			FeatureInfo[] features = this.packageManager.getSystemAvailableFeatures();
+			for (FeatureInfo feature : features) {
+				if (feature.name != null)
+				{
+					if (feature.name.equals("android.hardware.vulkan.level"))
+					{
+						// since we may not be compiled against android-24 or higher, use .toString to get the version field
+						String dump = feature.toString();
+						int index = dump.indexOf("v=");
+						if (index >= 0)
+						{
+							dump = dump.substring(index+2);
+							index = dump.indexOf(" ");
+							if (index >= 0)
+							{
+								ref.VulkanLevel = Integer.parseInt(dump.substring(0, index));
+								Log.d(TAG, "Vulkan level: " + ref.VulkanLevel);
+							}
+						}
+					}
+					else
+					if (feature.name.equals("android.hardware.vulkan.version"))
+					{
+						// since we may not be compiled against android-24 or higher, use .toString to get the version field
+						String dump = feature.toString();
+						int index = dump.indexOf("v=");
+						if (index >= 0)
+						{
+							dump = dump.substring(index+2);
+							index = dump.indexOf(" ");
+							if (index >= 0)
+							{
+								ref.VulkanVersion = Integer.parseInt(dump.substring(0, index));
+								int VersionMajor = (ref.VulkanVersion >> 22) & 0x03ff;
+								int VersionMinor = (ref.VulkanVersion >> 12) & 0x03ff;
+								int VersionPatch = ref.VulkanVersion & 0x0fff;
+								ref.VulkanVersionString = VersionMajor + "." + VersionMinor + "." + VersionPatch;
+								Log.d(TAG,"GetVulkanInfoFromPackage: Vulkan version: " + VersionMajor + "." + VersionMinor + "." + VersionPatch);
+							}
+						}
+					}
+				}
+			}
+		}
+	}
+
+	public GameActivitySetupInfo(Surface externalSurface, SimpleContextWrapper inActivity, String OBBFilename, boolean enablePropagateAlpha) {
+		if (inActivity == null)
+		{
+			return;
+		}
+
+		this.clipboardManager = (ClipboardManager) inActivity.getSystemService(CLIPBOARD_SERVICE);
+
+		this.packageManager = inActivity.getApplicationContext().getPackageManager();
+		
+		// Grab a reference to the asset manager
+		this.AssetManagerReference = inActivity.getAssets();
+
+		this.InternalFilesDir = inActivity.getFilesDir().getPath() + "/";
+		//this.InternalFilesDir = inActivity.getFilesDir().getAbsolutePath() + "/";
+		File _externalFilesDir = inActivity.getExternalFilesDir(null);
+		this.ExternalFilesDir = (_externalFilesDir != null ? _externalFilesDir.getPath() : "") + "/";
+		this.packageResourcePath = inActivity.getPackageResourcePath();
+
+		this.renderSurface = externalSurface;
+
+		if (inActivity != null)
+		{
+			Activity _AndroidActivity = inActivity.mAppActivityContext != null && inActivity.mAppActivityContext instanceof Activity ? ((Activity)inActivity.mAppActivityContext) : null;
+			this._extrasBundle = _AndroidActivity != null ? _AndroidActivity.getIntent().getExtras() : null;
+		}
+
+		if (this._extrasBundle == null)
+		{
+			this._extrasBundle = new Bundle();
+
+		}
+		if (this.renderSurface != null) {
+			this._extrasBundle.putParcelable("renderSurface", this.renderSurface);
+			Log.d(TAG, "====> GameActivitySetupInfo(PUT) is storing renderSurface = " + this.renderSurface);
+		}
+
+		this.appPackageName = inActivity.getApplicationContext().getPackageName(); //inActivity.getApplicationContext().getPackageName()
+		Log.w(TAG, "====> GameActivitySetupInfo() check OBBFilename = " + OBBFilename);
+
+		this.ProjectName = this.appPackageName.substring(this.appPackageName.lastIndexOf('.') + 1);
+
+
+
+		String AppType = "";
+		try
+		{
+			ApplicationInfo ai = inActivity.getPackageManager().getApplicationInfo(this.appPackageName, PackageManager.GET_META_DATA);
+			Bundle bundle = ai.metaData;
+			this._bundle = bundle;
+			PackageInfo packageInfo = inActivity.getPackageManager().getPackageInfo(this.appPackageName, 0);
+			this.VersionCode = packageInfo.versionCode;
+			this.targetSdkVersion = packageInfo.applicationInfo.targetSdkVersion;
+		}
+		catch (PackageManager.NameNotFoundException e)
+		{
+			Log.d(TAG, "Failed to load meta-data: NameNotFound: " + e.getMessage());
+		}
+		catch (NullPointerException e)
+		{
+			Log.d(TAG, "Failed to load meta-data: NullPointer: " + e.getMessage());
+		}
+
+		setOBBFilename(OBBFilename);
+
+		Log.v(TAG,"====> GameActivitySetupInfo(constructor) with ProjectName = " + this.ProjectName
+				+ ", appPackageName = " + this.appPackageName
+				+ ", _OBBFilename = " + this._OBBFilename
+				+ ", packageDataInsideApkValue = " + this.packageDataInsideApkValue
+		);
+
+
+		this.PropagateAlpha = enablePropagateAlpha ? 1 : 0;
+		this.OrientationPortrait = inActivity.getResources().getConfiguration().orientation == Configuration.ORIENTATION_PORTRAIT;
+		
+		// store the screen resolution (adjusted for immersive mode)
+		Map<String, String> variables = new HashMap<String, String>();
+		android.view.Display display = inActivity.getWindowManager().getDefaultDisplay();
+		android.graphics.Point displaySize = new android.graphics.Point();
+		
+		// check if immersive mode on window
+//		View decorView = inActivity.getWindow().getDecorView();
+//		if ((decorView.getSysteUiVisibility() & View.System_UI_FLAG_IMMERSIVE_STICKY) != 0)
+		{
+			display.getRealSize(displaySize);
+		}
+//		else
+//		{
+//			display.getSize(displaySize);
+//		}
+		variables.put("screenWidth", Integer.toString(displaySize.x));
+		variables.put("screenHeight", Integer.toString(displaySize.y));
+		
+		// send the variables to the engine
+		this.keyValues = new String[variables.size() * 2];
+		int keyValueIndex = 0;
+		for (Map.Entry<String, String> entry : variables.entrySet())
+		{
+			this.keyValues[keyValueIndex++] = entry.getKey();
+			this.keyValues[keyValueIndex++] = entry.getValue();
+		}
+	}
+
+
+
+	public void AssignContext(Context context)
+	{
+		AssetManagerReference = context.getAssets();
+		clipboardManager = (ClipboardManager) context.getSystemService(CLIPBOARD_SERVICE);
+		packageManager = context.getApplicationContext().getPackageManager();
+	}
+
+	public void DoNativeUpdates(GameActivityForMakeAAR gameActivity)
+	{
+		Log.d(TAG,"DoNativeUpdates() ExternalFilesDir path: " + this.ExternalFilesDir);
+
+		gameActivity.nativeSetGlobalActivity(this.UseExternalFilesDir, this.PublicLogFiles, this.InternalFilesDir, this.ExternalFilesDir, this.packageDataInsideApkValue, this.packageResourcePath);
+		if (this.packageDataInsideApkValue)
+		{
+			Log.d(TAG,"OBB override: " + this._OBBFilename);
+			gameActivity.nativeSetObbFilePaths(this._OBBFilename, "", "", "");
+		}
+
+		// tell the engine if this is a portrait app and if it will propagate alpha
+		gameActivity.nativeSetWindowInfo(this.OrientationPortrait, this.DepthBufferPreference, this.PropagateAlpha);
+
+		// get the full language code, like en-US
+		// note: this may need to be Locale.getDefault().getLanguage()
+		String Language = java.util.Locale.getDefault().toString();
+
+		boolean bDebuggerAttached = android.os.Debug.isDebuggerConnected();
+		gameActivity.nativeSetAndroidStartupState(bDebuggerAttached);
+
+		// Retrieve version code, name, and target SDK
+		try
+		{
+			PackageInfo packageInfo = gameActivity.getPackageManager().getPackageInfo(gameActivity.getPackageName(), 0);
+			VersionCode = packageInfo.versionCode;
+			gameActivity.VersionName = packageInfo.versionName;
+			gameActivity.targetSdkVersion = packageInfo.applicationInfo.targetSdkVersion;
+		}
+		catch (Exception e)
+		{
+			Log.d(TAG,"Error accessing packageInfo: " + e.getMessage());
+		}
+
+		String productName = gameActivity.getProductName();
+		Log.d(TAG, "Android version is " + android.os.Build.VERSION.RELEASE );
+		Log.d(TAG, "Android manufacturer is " + android.os.Build.MANUFACTURER );
+		Log.d(TAG, "Android model is " + android.os.Build.MODEL );
+		Log.d(TAG, "Android build number is " + android.os.Build.DISPLAY );
+		Log.d(TAG, "OS language is set to " + Language);
+		Log.d(TAG, "Product name is " + (productName.isEmpty() ? "[not set]" : productName));
+		Log.d(TAG, "Debugger attached is " + bDebuggerAttached );
+		Log.d(TAG, "Android targetSdkVersion version is " + gameActivity.targetSdkVersion );
+		
+		gameActivity.nativeSetAndroidVersionInformation(android.os.Build.VERSION.RELEASE, gameActivity.targetSdkVersion, android.os.Build.MANUFACTURER, android.os.Build.MODEL, android.os.Build.DISPLAY, Language, productName );
+
+		try
+		{
+			int PatchVersion = 0;
+			String AppType = "";
+			gameActivity.nativeSetObbInfo(this.ProjectName, this.appPackageName, this.VersionCode, PatchVersion, AppType);
+		}
+		catch (Exception e)
+		{
+			// if the above failed, then, we can't use obbs
+			Log.d(TAG,"==================================> PackageInfo failure getting .obb info: " + e.getMessage());
+		}
+
+		// tell Android that we want volume controls to change the media volume, aka music
+		// enable the physical volume controls to the game
+//			_activity.setVolumeControlStream(AudioManager.STREAM_MUSIC);
+
+
+		gameActivity.nativeSetConfigRulesVariables(this.keyValues);
+	}
+
+	private void setOBBFilename(String obbFilename)
+	{
+		this._OBBFilename = obbFilename;
+		this.packageDataInsideApkValue = obbFilename.equals("") ? true : false;
+	}
+
+//	@RequiresApi(api = Build.VERSION_CODES.Q)
+	protected GameActivitySetupInfo(Parcel in)
+	{
+		this.InternalFilesDir = in.readString();
+		this.ExternalFilesDir = in.readString();
+		this.packageResourcePath = in.readString();
+		this.OrientationPortrait = in.readBoolean();
+		this.DepthBufferPreference = in.readInt();
+		this.PropagateAlpha = in.readInt();
+		this.UseExternalFilesDir = in.readBoolean();
+		this.PublicLogFiles = in.readBoolean();
+
+		this._OBBFilename = in.readString();
+		this.packageDataInsideApkValue = in.readBoolean();
+		in.readIntArray(this.ViewPosition);
+		in.readIntArray(this.ViewSize);
+
+		this.VersionCode = in.readInt();
+		this.targetSdkVersion = in.readInt();
+
+		this._extrasBundle = in.readBundle(getClass().getClassLoader());
+		if (this._extrasBundle == null)
+		{
+			Log.e(TAG,"=====> PackageInfo _extrasBundle so no renderSurface is passed!" );
+		}
+		this.renderSurface = this._extrasBundle == null ? null : this._extrasBundle.<Surface>getParcelable("renderSurface");
+		Log.d(TAG, "====> GameActivitySetupInfo(GET) is storing renderSurface = " + this.renderSurface);
+
+		Log.e(TAG,"PackageInfo from _extrasBundle =====>  renderSurface = " + this.renderSurface );
+
+		this._bundle = in.readBundle(getClass().getClassLoader());
+		this.appPackageName = in.readString();
+		this.ProjectName = in.readString();
+		this.IsForDistribution = in.readBoolean();
+		this.keyValues = new String[4];
+
+		in.readStringArray(this.keyValues);
+
+	}
+
+	public final static Creator<GameActivitySetupInfo> CREATOR = new Creator<GameActivitySetupInfo>() {
+//		@RequiresApi(api = Build.VERSION_CODES.Q)
+		@Override
+		public GameActivitySetupInfo createFromParcel(Parcel in) {
+			return new GameActivitySetupInfo(in);
+		}
+
+		@Override
+		public GameActivitySetupInfo[] newArray(int size) {
+			return new GameActivitySetupInfo[size];
+		}
+	};
+
+	@Override
+	public int describeContents() {
+		return 0;
+	}
+
+//	@RequiresApi(api = Build.VERSION_CODES.Q)
+	@Override
+	public void writeToParcel(Parcel parcel, int i) {
+
+		parcel.writeString(this.InternalFilesDir);
+		parcel.writeString(this.ExternalFilesDir);
+		parcel.writeString(this.packageResourcePath);
+		parcel.writeBoolean(this.OrientationPortrait);
+		parcel.writeInt(this.DepthBufferPreference);
+		parcel.writeInt(this.PropagateAlpha);
+		parcel.writeBoolean(this.UseExternalFilesDir);
+		parcel.writeBoolean(this.PublicLogFiles);
+
+		parcel.writeString(this._OBBFilename);
+		parcel.writeBoolean(this.packageDataInsideApkValue);
+
+		parcel.writeIntArray(this.ViewPosition);
+		parcel.writeIntArray(this.ViewSize);
+
+		parcel.writeInt(this.VersionCode);
+		parcel.writeInt(this.targetSdkVersion);
+
+		parcel.writeBundle(this._extrasBundle);
+		parcel.writeBundle(this._bundle);
+		parcel.writeString(this.appPackageName);
+		parcel.writeString(this.ProjectName);
+		parcel.writeBoolean(this.IsForDistribution);
+		parcel.writeStringArray(this.keyValues);
+	}
+}
+
diff --git a/Engine/Plugins/Runtime/AndroidSingleInstanceService/Source/AndroidSingleInstanceService/Android/src/main/java/com/epicgames/makeaar/UnrealSharedInstanceService.java b/Engine/Plugins/Runtime/AndroidSingleInstanceService/Source/AndroidSingleInstanceService/Android/src/main/java/com/epicgames/makeaar/UnrealSharedInstanceService.java
new file mode 100644
--- /dev/null
+++ b/Engine/Plugins/Runtime/AndroidSingleInstanceService/Source/AndroidSingleInstanceService/Android/src/main/java/com/epicgames/makeaar/UnrealSharedInstanceService.java
@@ -0,0 +1,481 @@
+package com.epicgames.makeaar;
+
+import android.app.Activity;
+import android.app.Application;
+import android.app.Service;
+import android.content.Context;
+import android.content.Intent;
+import android.content.res.Configuration;
+import android.graphics.Point;
+import android.net.Uri;
+import android.os.Bundle;
+import android.os.Handler;
+import android.os.IBinder;
+import android.os.Message;
+import android.os.Messenger;
+import android.os.RemoteException;
+import android.provider.Settings;
+import android.util.Log;
+import android.view.Surface;
+import android.view.TextureView;
+import android.widget.Toast;
+
+import com.epicgames.makeaar.Engine;
+import com.epicgames.makeaar.EngineFactory;
+import com.epicgames.makeaar.GameActivityForMakeAAR;
+import com.epicgames.makeaar.GameActivitySetupInfo;
+import com.epicgames.unreal.SimpleContextWrapper;
+
+import com.epicgames.makeaar.UnrealMessageType;
+
+import java.lang.ref.PhantomReference;
+import java.lang.ref.WeakReference;
+
+public class UnrealSharedInstanceService extends Service {
+
+	private static final String TAG = "UE_ASIS_SharedInstanceService";
+
+	static int activeTaskID = 0;
+	/**
+	* Target we publish for clients to send messages to IncomingHandler.
+	*/
+	Messenger mMessenger;
+	public String Name()
+	{
+		return getApplication().getClass().getSimpleName();
+	}
+	public static String Name(Context context)
+	{
+		return context.getPackageName() +",context=" + context.toString();
+	}
+
+	public UnrealSharedInstanceService() {
+	}
+
+	/**
+	 * Handler of incoming messages from clients.
+	 */
+	static class IncomingHandler extends Handler implements Engine.IEventCallback
+	{
+		private Service applicationContext;
+		static private Engine engineInstance;
+
+
+		public static String defaultProjectModuleName = "AndroidPackagingTest";
+//		public static String defaultProjectModuleName = "Alpha_POC";
+		public static String defaultOBBFileLocation = "";//"Gallery3D"; //"Alpha_POC";
+		//	public final static String defaultProjectModuleName = "";
+//		public static String defaultCommandLine = "../../../" + BuildConfig.OBB_MODULE_NAME + "/" + BuildConfig.OBB_MODULE_NAME + ".uproject -nosound -nocrashreports";
+		public static String defaultCommandLine = "-nosound -nocrashreports -dpcvars=a.UseSwappyForFramePacing=0 -tracehost=127.0.0.1 -cpuprofilertrace -statnamedevents -trace=Bookmark,Frame,CPU,GPU,LoadTime,File";
+//		public final static Engine gUnrealEngine = EngineFactory.getInstance(null, defaultOBBFileLocation, defaultProjectModuleName, true);
+//		static Engine gUnrealEngine;
+		public static Messenger mServiceReply;
+
+
+
+
+		void Release()
+		{
+			engineInstance.registerEventCallback(null);
+			mServiceReply = null;
+		}
+
+		final static void SetOBB(String commandLineArgs, String obbModuleName, String obbFileLocation)
+		{
+			if (commandLineArgs.isEmpty())
+				commandLineArgs = "-nosound -nocrashreports";
+			if (!obbModuleName.isEmpty()) {
+				defaultProjectModuleName = obbModuleName;
+				defaultCommandLine = "../../../" + defaultProjectModuleName + "/" + defaultProjectModuleName + ".uproject " + commandLineArgs;
+//			if (!obbFileLocation.isEmpty())
+				defaultOBBFileLocation = obbFileLocation;
+			}
+
+		}
+
+		private static int APP_CMD_INPUT_CHANGED = 0;
+		private static int APP_CMD_INIT_WINDOW = 1;
+		private static int APP_CMD_TERM_WINDOW = 2;
+		private static int APP_CMD_WINDOW_RESIZED = 3;
+		private static int APP_CMD_WINDOW_REDRAW_NEEDED = 4;
+		private static int APP_CMD_CONTENT_RECT_CHANGED = 5;
+		private static int APP_CMD_GAINED_FOCUS = 6;
+		private static int APP_CMD_LOST_FOCUS = 7;
+		private static int APP_CMD_CONFIG_CHANGED = 8;
+		private static int APP_CMD_LOW_MEMORY = 9;
+		private static int APP_CMD_START = 10;
+		private static int APP_CMD_RESUME = 11;
+		private static int APP_CMD_SAVE_STATE = 12;
+		private static int APP_CMD_PAUSE = 13;
+		private static int APP_CMD_STOP = 14;
+		private static int APP_CMD_DESTROY = 15;
+		
+		IncomingHandler(Service context, Intent intent) {
+		    mServiceReply = intent.getParcelableExtra("callback");
+			applicationContext = context;
+
+		    if (intent.hasExtra("obbModuleName"))
+		    {
+			    Bundle data = intent.getExtras();
+			    String obbModuleName = data.getString("obbModuleName", "");
+			    String obbFileLocation = data.getString("obbFileLocation", "");
+			    String commandLineArgs = data.getString("commandLineArgs", "");
+			    SetOBB(commandLineArgs, obbModuleName, obbFileLocation);
+    
+			    Log.i("UESharedInstanceServiceNative", "**\tIncomingHandler(constructor intent) - proc=" + Name(context)
+				    + ", obbModuleName=" + obbModuleName
+				    + ", obbFileLocation=" + obbFileLocation);
+		    }
+		    
+		    if (engineInstance == null) {
+			    engineInstance = EngineFactory.getInstance(null, defaultOBBFileLocation, defaultProjectModuleName, true);
+		    }
+		    
+		    engineInstance.registerEventCallback(this);
+		    engineInstance.allowConsole(true);
+		}
+
+		@Override
+		public void eventCallback(int event, String param1, int param2, int param3, float param4)
+		{
+
+			if (event != 3 && event != 4) {
+				Log.d(TAG, "eventCallback: " + event);
+			}
+
+			if (event == Engine.EVENTTYPE_FRAME_BEGIN) {
+
+
+			}
+
+			if (event == Engine.EVENTTYPE_FRAME_END) {
+
+			}
+
+			if (event == Engine.EVENTTYPE_POST_ENGINE_INIT)
+			{
+
+
+			}
+
+			if (event == Engine.EVENTTYPE_PRE_LOAD_MAP)
+			{
+
+			}
+
+			if (event == Engine.EVENTTYPE_POST_LOAD_MAP) {
+
+				//Log.i(TAG, "Configure unlocked FPS");
+				//engine.sendConsoleCommand("rhi.SyncInterval=1"); // 0 = unlocked, 1 = 60 fps, 2 = 30 fps
+				//engine.sendConsoleCommand("r.Vsync=0");
+				//engine.sendConsoleCommand("r.GTSyncType=1"); // will cause the game thread to sync with the RHI thread (rather than the rendering thread), which caps the upper limit of input latency to something more sensible
+
+
+
+			} else if (event == Engine.EVENTTYPE_ACTION) {
+
+				Log.d(TAG, "eventCallback: " + event);
+			}
+
+			if (mServiceReply != null) {
+				try {
+					Message replyMsg = Message.obtain();
+					replyMsg.what = event;
+					if (param1 != null && !param1.isEmpty())
+					{
+//						Log.d(TAG, "writing parceable eventCallback: " + event + ", param1=" + param1);
+						replyMsg.getData().putString("param1", param1);
+
+//						replyMsg.obj = new UEReplyEvent(event, param1, param2, param3, param4);
+					}
+					replyMsg.arg1 = param2;
+					replyMsg.arg2 = param3;
+
+					if (param4 != 0.0f)
+					{
+//						Log.d(TAG, "writing parceable eventCallback: " + event + ", param4=" + param4);
+						replyMsg.getData().putFloat("param4", param4);
+					}
+					mServiceReply.send(replyMsg);
+				} catch (RemoteException e) {
+					e.printStackTrace();
+				}
+			}
+		}
+
+		@Override
+		public void handleMessage(Message msg) {
+
+
+			if (msg.what == UnrealMessageType.Hello.ordinal()) {
+				Toast.makeText(applicationContext, "hello!", Toast.LENGTH_SHORT).show();
+			}
+			else if (msg.what == UnrealMessageType.AttachExternalSurface.ordinal())
+			{
+				Bundle data = msg.getData();
+				Surface externalSurface = data.getParcelable("surface");
+				int taskId = data.getInt("taskId", 0);
+				boolean bDoStart = data.getBoolean("start", true);
+				boolean bDoResume = data.getBoolean("resume", true);
+
+				Log.i(TAG, "handleMessage(MSG_ATTACH_EXTERNAL_SURFACE) - proc=" + Name(this.applicationContext)
+					+ ", taskId=" + taskId
+					+ ", activeTaskID=" + activeTaskID
+					+ ", externalSurface=" + externalSurface
+					+ ", applicationContext=" + applicationContext);
+
+				int[] emptyViewSize = {0,0};
+				int[] viewSize = emptyViewSize;
+				int[] viewPos = emptyViewSize;
+
+				if (data.containsKey("viewSize"))
+				{
+					viewSize = data.getIntArray("viewSize");
+				}
+				if (data.containsKey("viewPosition")) {
+					viewPos = data.getIntArray("viewPosition");
+				}
+				boolean overridePropagateAlpha = data.containsKey("enablePropagateAlpha");
+				boolean enablePropagateAlpha = data.getBoolean("enablePropagateAlpha", false);
+
+				engineInstance.registerEventCallback(this);
+
+				//if (bDoStart)
+				{
+					//engineInstance.QueueStart("UnrealMessageType.AttachExternalSurface, taskId=" + taskId); //777
+				}
+
+				if (externalSurface != null && externalSurface.isValid() && engineInstance != null) {
+					if (engineInstance.IsInitialized() == false) {
+						SimpleContextWrapper simpleActivity = new SimpleContextWrapper(applicationContext);
+
+						GameActivitySetupInfo gameActivitySetupInfo = new GameActivitySetupInfo(externalSurface, simpleActivity, defaultOBBFileLocation, enablePropagateAlpha);
+						gameActivitySetupInfo.ViewPosition = viewPos;
+						gameActivitySetupInfo.ViewSize = viewSize;
+
+						bDoResume &= engineInstance.Init(simpleActivity, gameActivitySetupInfo, defaultCommandLine, defaultProjectModuleName, new String[]{});
+					} else {
+						if (overridePropagateAlpha) {
+							final int DepthBufferPreference = 32;
+							final int PropagateAlpha = enablePropagateAlpha ? 1 : 0;
+							boolean bPortrait = applicationContext.getResources().getConfiguration().orientation == Configuration.ORIENTATION_PORTRAIT;
+
+							GameActivityForMakeAAR.Get().nativeSetWindowInfo(bPortrait, DepthBufferPreference, PropagateAlpha);
+						}
+
+						bDoResume &= engineInstance.AttachExternalRenderSurface(taskId, externalSurface, viewPos, viewSize);
+					}
+
+					
+					//if (bDoResume)
+					{
+						engineInstance.QueueResume("UnrealMessageType.AttachExternalSurface, taskId=" + taskId);
+					}
+
+					activeTaskID = taskId;
+
+					if (msg.replyTo != null) {
+						Log.w(TAG, "Sending message for engine init back to message replyTo obj=" + engineInstance);
+						Message replyMsg = Message.obtain();
+						replyMsg.what = Engine.EVENTTYPE_POST_ENGINE_INIT;
+						msg.arg1 = activeTaskID;
+//							msg.obj = localBinder;
+//							Bundle rdata = new Bundle();
+//							rdata.putBundle("key", engineInstance);
+//							msg.setData(); = engineInstance;
+
+//						    try {
+//							    mServiceReply.send(replyMsg);
+//						    } catch (RemoteException e) {
+//							    e.printStackTrace();
+//						    }
+					}
+					else {
+						Log.w(TAG, "Sending message for engine init using localBinder.handler=" + this);
+
+//						    Message replyMsg = this.obtainMessage(Engine.EVENTTYPE_POST_ENGINE_INIT, engineInstance);
+//						    this.sendMessage(replyMsg);
+					}
+				}
+			}
+			else if (msg.what == UnrealMessageType.DetachExternalSurface.ordinal())
+			{
+				Bundle data = msg.getData();
+
+				int taskId = data.getInt("taskId", 0);
+				boolean bDoStop = data.getBoolean("stop", true);
+				boolean bDoPause = data.getBoolean("pause", true);
+
+				Log.i(TAG, "handleMessage(MSG_DETACH_EXTERNAL_SURFACE) - proc=" + Name(applicationContext)
+					+ ", taskId=" + taskId
+					+ ", activeTaskID=" + activeTaskID
+					+ ", applicationContext=" + applicationContext);
+
+				//if (taskId != 0 && taskId == activeTaskID )
+				{
+
+					//if (bDoPause)
+					{
+						engineInstance.QueuePause("UnrealMessageType.DetachExternalSurface, taskId=" + taskId);
+					}
+
+					//if (bDoStop)
+					{
+						//engineInstance.QueueStop("UnrealMessageType.DetachExternalSurface, taskId=" + taskId);
+					}
+
+					//engineInstance.sendConsoleCommand("t.maxfps 0.001");
+					
+				}
+				activeTaskID = taskId;
+
+				//mServiceReply = null;
+				//engineInstance.registerEventCallback(null);
+				//applicationContext.stopSelf();
+
+			}
+			else if (msg.what == UnrealMessageType.StopService.ordinal())
+			{
+				Bundle data = msg.getData();
+				int taskId = data.getInt("taskId", 0);
+
+				Log.i(TAG, "handleMessage(MSG_STOP_SERVICE) - proc=" + Name(applicationContext)
+					+ ", taskId=" + taskId
+					+ ", activeTaskID=" + activeTaskID
+					+ ", applicationContext=" + applicationContext);
+
+				engineInstance.QueuePause("UnrealMessageType.StopService, taskId=" + taskId);
+				engineInstance.QueueStop("UnrealMessageType.StopService, taskId=" + taskId);
+				//engineInstance.onPause(engineInstance.getCurrentContextID(), "UnrealMessageType.StopService, taskId=" + taskId);
+				//engineInstance.onStop(engineInstance.getCurrentContextID(), "UnrealMessageType.StopService, taskId=" + taskId);
+
+				activeTaskID = taskId;
+
+				mServiceReply = null;
+				engineInstance.registerEventCallback(null);
+				applicationContext.stopSelf();
+			}
+			else if (msg.what == UnrealMessageType.ResumeService.ordinal())
+			{
+				Bundle data = msg.getData();
+				int taskId = data.getInt("taskId", 0);
+				boolean bResume = data.getBoolean("resume", true);
+				boolean bCMDMainInit = data.getBoolean("cmd_mainInit", false);
+
+				Log.i(TAG, "handleMessage(MSG_RESUME_SERVICE) - proc=" + Name(applicationContext)
+					+ ", taskId=" + taskId
+					+ ", activeTaskID=" + activeTaskID
+					+ ", applicationContext=" + applicationContext);
+
+				engineInstance.registerEventCallback(this);
+
+				if (bResume)
+				{
+					engineInstance.QueueResume("UnrealMessageType.ResumeService, taskId=" + taskId);
+				}
+				else {
+
+					if (bCMDMainInit)
+						GameActivityForMakeAAR.Get().nativeResumeMainInit();
+				}
+
+			}
+			else if (msg.what == UnrealMessageType.TouchEvent.ordinal())
+			{
+				engineInstance.onTouchEvent(msg.getData().getParcelable("touch"));
+			}
+			else if (msg.what == UnrealMessageType.SendConsoleCommand.ordinal())
+			{
+				Bundle data = msg.getData();
+				String consoleCommand = data.getString("consoleCommand");
+
+				Log.i(TAG, "handleMessage(MSG_ISSUE_CONSOLE_COMMAND) - proc=" + Name(applicationContext)
+						+ ", consoleCommand=" + consoleCommand
+						+ ", activeTaskID=" + activeTaskID
+						+ ", applicationContext=" + applicationContext);
+
+				engineInstance.sendConsoleCommand(consoleCommand);
+			}
+			else
+			{
+				super.handleMessage(msg);
+			}
+		}
+
+		@Override
+		protected void finalize() throws Throwable {
+			super.finalize();
+			mServiceReply = null;
+			engineInstance.registerEventCallback(null);
+		}
+	}
+
+
+	
+	@Override
+    public void onCreate() {
+        super.onCreate();
+		Log.i(TAG, "onCreate(UnrealSharedInstanceService) called. context = " + this 
+			+ ", activeTaskID=" + activeTaskID 
+		);
+	}
+
+
+    // execution of service will start
+    // on calling this method
+	@Override
+	public int onStartCommand(Intent intent, int flags, int startId) {
+
+		Log.i(TAG, "UnrealSharedInstanceService::onStartCommand called. context = " + this + ", startID=" + startId );
+
+		// returns the status
+		// of the program
+		return START_STICKY;
+    }
+	
+
+    // execution of the service will
+    // stop on calling this method
+	//
+    @Override
+    public void onDestroy() {
+        super.onDestroy();
+
+		Log.w(TAG, "onDestroy(UnrealAndroidInstanceService) - proc = " + Name(this)
+			+ ", mMessenger = " + mMessenger
+		);
+
+		if (mMessenger != null) {
+			try {
+				mMessenger.send(Message.obtain(null, UnrealMessageType.StopService.ordinal()));
+			} catch (RemoteException e) {
+				e.printStackTrace();
+			}
+		}
+		mMessenger = null;
+    }
+
+    @Override
+    public IBinder onBind(Intent intent) {
+		int taskId = intent.getIntExtra("taskId", 0);
+
+		Log.i(TAG, "onBind(intent) - proc=" + Name(this)
+			+ ", taskId=" + taskId
+			+ ", activeTaskID=" + activeTaskID
+			+ ", intent=" + intent);
+
+
+		Toast.makeText(getApplicationContext(), "binding taskId=" + taskId, Toast.LENGTH_SHORT).show();
+		mMessenger = new Messenger(new IncomingHandler(this, intent));
+		return mMessenger.getBinder();
+    }
+	@Override
+	public boolean onUnbind(Intent intent) {
+		int taskId = intent.getIntExtra("taskId", 0);
+		Log.i(TAG, "onUnbind(intent) - proc=" + Name(this)
+			+ ", taskId=" + taskId
+			+ ", activeTaskID=" + activeTaskID
+			+ ", intent=" + intent);
+		return super.onUnbind(intent);
+	}
+}
\ No newline at end of file
diff --git a/Engine/Plugins/Runtime/AndroidSingleInstanceService/Source/AndroidSingleInstanceService/AndroidSingleInstanceService.Build.cs b/Engine/Plugins/Runtime/AndroidSingleInstanceService/Source/AndroidSingleInstanceService/AndroidSingleInstanceService.Build.cs
new file mode 100644
--- /dev/null
+++ b/Engine/Plugins/Runtime/AndroidSingleInstanceService/Source/AndroidSingleInstanceService/AndroidSingleInstanceService.Build.cs
@@ -0,0 +1,39 @@
+// Copyright Epic Games, Inc. All Rights Reserved.
+
+using UnrealBuildTool;
+using System.IO;
+
+public class AndroidSingleInstanceService : ModuleRules
+{
+	public AndroidSingleInstanceService(ReadOnlyTargetRules Target) : base(Target)
+	{
+		PCHUsage = ModuleRules.PCHUsageMode.UseExplicitOrSharedPCHs;
+		
+		PublicDependencyModuleNames.AddRange(
+			new string[]
+			{
+				"Core",
+			}
+			);
+			
+		
+		PrivateDependencyModuleNames.AddRange(
+			new string[]
+			{
+				"CoreUObject",
+				"Engine",
+				"Slate",
+				"SlateCore",
+			}
+			);
+		
+		
+		if (Target.Platform == UnrealTargetPlatform.Android)
+		{
+			string PluginPath = Utils.MakePathRelativeTo(ModuleDirectory, Target.RelativeEnginePath);
+			AdditionalPropertiesForReceipt.Add("AndroidPlugin", Path.Combine(PluginPath, "AndroidSingleInstanceService_UPL.xml"));
+
+			PrivateDependencyModuleNames.AddRange(new string[] { "Launch" });
+		}
+	}
+}
diff --git a/Engine/Plugins/Runtime/AndroidSingleInstanceService/Source/AndroidSingleInstanceService/AndroidSingleInstanceService_UPL.xml b/Engine/Plugins/Runtime/AndroidSingleInstanceService/Source/AndroidSingleInstanceService/AndroidSingleInstanceService_UPL.xml
new file mode 100644
--- /dev/null
+++ b/Engine/Plugins/Runtime/AndroidSingleInstanceService/Source/AndroidSingleInstanceService/AndroidSingleInstanceService_UPL.xml
@@ -0,0 +1,202 @@
+<?xml version="1.0" encoding="utf-8"?>
+<root xmlns:android="http://schemas.android.com/apk/res/android" xmlns:dist="http://schemas.android.com/apk/distribution">
+	<init>
+		<log text="AndroidSingleInstanceService Plugin Init"/>
+		<trace enable="true"/>
+		<setBoolFromProperty result="bEnabled" ini="Engine" section="/Script/AndroidSingleInstanceServiceEditor.AndroidSingleInstanceServiceRuntimeSettings" property="bEnablePlugin" default="false"/>
+		<setBoolFromProperty result="bCompileASISProject" ini="Engine" section="/Script/AndroidSingleInstanceServiceEditor.AndroidSingleInstanceServiceRuntimeSettings" property="bCompileASISProject" default="false"/>
+
+		<log text="AndroidSingleInstanceService AndroidPackageName = $S(AndroidPackageName)"/>
+		
+		<!-- create the ASIS project if plugin enabled -->
+		<setBool result="bCreateASISProject" value="$B(bEnabled)"/>
+		
+		<!-- by default embed ASIS in APK if plugin enabled -->
+		<setBool result="bEmbedASIS" value="$B(bEnabled)"/>
+
+		<!-- by default external start of ASIS allowed if plugin enabled -->
+		<setBool result="bExternalStartASIS" value="$B(bEnabled)"/>
+			 
+		<if condition="bEnabled">
+			<false>
+				<setBool result="bCreateASISProject" value="false"/>
+			</false>
+		</if>
+
+
+
+	
+	</init>
+
+	<!-- register the ini keys that modify the packaging -->
+	<registerBuildSettings>
+		<insertValue value="/Script/AndroidSingleInstanceServiceEditor.AndroidSingleInstanceServiceRuntimeSettings[bEnablePlugin]"/>
+		<insertNewline/>
+	</registerBuildSettings>
+
+	<baseBuildGradleAdditions>
+		<if condition="bEnabled">
+			<true>
+				<log text="AndroidSingleInstanceService baseBuildGradleAdditions bEnabled. for BuildDir= $S(BuildDir)"/>
+				<addElements tag="buildscript">
+					<insertValue value="apply from: 'buildsystem/config.gradle'"/>
+					<insertValue value="apply from: 'buildsystem/dependencies.gradle'"/>
+					<insertNewline/>
+				</addElements>
+			</true>
+		</if>
+	</baseBuildGradleAdditions>
+	
+	
+	<buildscriptGradleAdditions>
+	</buildscriptGradleAdditions>
+
+	<androidManifestUpdates>
+		<if condition="bEnabled">
+			<true>
+				<log text="AndroidSingleInstanceService Manifest update application block..."/>
+
+				<log text="AndroidSingleInstanceService Manifest patch application attributes..."/>
+				<addAttribute tag="application" name="android:allowTaskReparenting" value="true" />
+				<addAttribute tag="application" name="android:hasCode" value="true" />
+				
+				<if condition="bEmbedASIS">
+					<true>
+						<log text="AndroidSingleInstanceService Manifest check bExternalStartASIS..."/>
+						<if condition="bExternalStartASIS">
+							<true>
+								<log text="AndroidSingleInstanceService Manifest bExternalStartASIS so add intent filters and service..."/>
+								<addElements tag="application">
+									<service android:name="com.epicgames.makeaar.UnrealSharedInstanceService" android:exported="true" android:process=":unrealSharedInstanceService" >
+									</service>
+								</addElements>
+							</true>
+						</if>
+					</true>
+				</if>
+				
+				<!-- force bExternalFilesDir to true for all builds -->
+				<setBool result="bUpdatedExternal" value="false"/>
+				<loopElements tag="meta-data">
+					<setStringFromAttribute result="metaName" tag="$" name="android:name" />
+					<setBoolIsEqual result="bIsExternal" arg1="$S(metaName)" arg2="com.epicgames.unreal.GameActivity.bUseExternalFilesDir" />
+					<if condition="bIsExternal">
+						<true>
+							<removeAttribute tag="$" name="android:value"/>
+							<addAttribute tag="$" name="android:value" value="true"/>
+							<setBool result="bUpdatedExternal" value="true"/>
+						</true>
+					</if>
+				</loopElements>
+				<if condition="bUpdatedExternal">
+					<false>
+						<addElements tag="application">
+							<meta-data android:name="com.epicgames.unreal.GameActivity.bUseExternalFilesDir" android:value="true"/>
+						</addElements>
+					</false>
+				</if>
+			<log text="AndroidSingleInstanceService Manifest patch done!"/>
+			</true>
+		</if>
+	
+	</androidManifestUpdates>
+
+	<gradleProperties>
+		<if condition="bEnabled">
+			<true>
+			</true>
+		</if>
+	</gradleProperties>
+
+	<resourceCopies>
+		<if condition="bEnabled">
+			<true>
+				<log text="AndroidSingleInstanceService resourceCopies copy service helper files. for BuildDirGradle= $S(BuildDir)/gradle/"/>
+				<copyDir src="$S(PluginDir)/Android/buildsystem" dst="$S(BuildDir)/gradle/buildsystem" />
+				<copyFile src="$S(PluginDir)/Android/src/main/java/com/epicgames/unreal/UEInterfaceHelpers.java" dst="$S(BuildDir)/src/com/epicgames/unreal/UEInterfaceHelpers.java" />
+			</true>
+		</if>
+			
+		<if condition="bEmbedASIS">
+			<true>
+				<log text="AndroidSingleInstanceService resourceCopies copy service class files..."/>
+				<deleteFiles filespec="$S(BuildDir)/src/com/epicgames/makeaar" />
+				<copyDir src="$S(PluginDir)/Android/src/main/java/com/epicgames/makeaar" dst="$S(BuildDir)/src/com/epicgames/makeaar" />
+				<copyFile src="$S(PluginDir)/Android/src/main/java/com/epicgames/makeaar/UnrealSharedInstanceService.java" dst="$S(BuildDir)/src/com/epicgames/makeaar/UnrealSharedInstanceService.java" />
+
+			</true>
+			<false>
+				<log text="AndroidSingleInstanceService resourceCopies bEmbedASIS is false so remove files..."/>
+				<deleteFiles filespec="$S(BuildDir)/src/com/epicgames/makeaar" />
+				<deleteFiles filespec="$S(BuildDir)/src/com/epicgames/makeaar/UnrealSharedInstanceService*.java" />
+			</false>
+		</if>
+	</resourceCopies>
+
+	<gradleCopies>
+		<if condition="bCreateASISProject">
+			<true>
+				<log text="AndroidSingleInstanceService create ASISStub project..."/>
+				<log text="AndroidSingleInstanceService bCreateASISProject, AndroidPackageName = $S(AndroidPackageName)"/>
+			</true>
+		</if>
+	</gradleCopies>
+	
+	<buildGradleAdditions>
+		<if condition="bCompileASISProject">
+			<true>
+				<log text="AndroidSingleInstanceService buildGradleAdditions bCompileASISProject..."/>
+			</true>
+		</if>
+
+		<if condition="bEmbedASIS">
+			<true>
+				<log text="AndroidSingleInstanceService buildGradleAdditions bEmbedASIS. for BuildDir= $S(BuildDir)"/>
+				<insert>
+					<![CDATA[
+
+apply from: '../buildsystem/config.gradle'
+apply from: '../buildsystem/dependencies.gradle'
+
+		android {
+		defaultConfig {
+		buildConfigField "String", "OBB_MODULE_NAME", config.obbModuleName
+		buildConfigField "String", "OBB_FILE_NAME", config.obbFileLocation
+		buildConfigField "String", "AAR_VERSION", config.AAR_VERSION
+		}
+		}
+]]>
+				</insert>
+			</true>
+		</if>
+	</buildGradleAdditions>
+	
+	<settingsGradleAdditions>
+		<if condition="bCreateASISProject">
+			<true>
+			</true>
+		</if>
+	</settingsGradleAdditions>
+
+	<gameActivityImportAdditions>
+		<if condition="bEmbedASIS">
+			<true>
+			</true>
+		</if>
+	</gameActivityImportAdditions>
+
+	<gameActivityClassAdditions>
+		<if condition="bEmbedASIS">
+			<true>
+				<insert>
+				</insert>
+				<insertNewline/>
+				<insert>
+				</insert>
+			</true>
+			<false>
+			</false>
+		</if>
+	</gameActivityClassAdditions>
+
+</root>
diff --git a/Engine/Plugins/Runtime/AndroidSingleInstanceService/Source/AndroidSingleInstanceService/Classes/AndroidSingleInstanceServiceBPLibrary.h b/Engine/Plugins/Runtime/AndroidSingleInstanceService/Source/AndroidSingleInstanceService/Classes/AndroidSingleInstanceServiceBPLibrary.h
new file mode 100644
--- /dev/null
+++ b/Engine/Plugins/Runtime/AndroidSingleInstanceService/Source/AndroidSingleInstanceService/Classes/AndroidSingleInstanceServiceBPLibrary.h
@@ -0,0 +1,207 @@
+// Copyright Epic Games, Inc. All Rights Reserved.
+
+#pragma once
+
+#include "CoreMinimal.h"
+#include "UObject/ObjectMacros.h"
+#include "Kismet/BlueprintFunctionLibrary.h"
+#include "AndroidSingleInstanceServiceBPLibrary.generated.h"
+
+DECLARE_DYNAMIC_MULTICAST_DELEGATE_OneParam(FButtonPressedMC, FString, Name);
+DECLARE_DYNAMIC_MULTICAST_DELEGATE_FiveParams(FAndroidEventReceivedMC, int32, EventId, FString, Param1, int32, Param2, int32, Param3, float, Param4);
+
+DECLARE_DYNAMIC_MULTICAST_DELEGATE_ThreeParams(FAndroidDataReceivedStringMC, int32, DataId, FString, Key, FString, Data);
+DECLARE_DYNAMIC_MULTICAST_DELEGATE_ThreeParams(FAndroidDataReceivedStringArrayMC, int32, DataId, FString, Key, const TArray<FString>&, Data);
+
+DECLARE_DYNAMIC_MULTICAST_DELEGATE_ThreeParams(FAndroidDataReceivedBooleanMC, int32, DataId, FString, Key, bool, Data);
+DECLARE_DYNAMIC_MULTICAST_DELEGATE_ThreeParams(FAndroidDataReceivedBooleanArrayMC, int32, DataId, FString, Key, const TArray<bool>&, Data);
+
+DECLARE_DYNAMIC_MULTICAST_DELEGATE_ThreeParams(FAndroidDataReceivedByteMC, int32, DataId, FString, Key, int8, Data);
+DECLARE_DYNAMIC_MULTICAST_DELEGATE_ThreeParams(FAndroidDataReceivedByteArrayMC, int32, DataId, FString, Key, const TArray<int8>&, Data);
+
+DECLARE_DYNAMIC_MULTICAST_DELEGATE_ThreeParams(FAndroidDataReceivedIntMC, int32, DataId, FString, Key, int32, Data);
+DECLARE_DYNAMIC_MULTICAST_DELEGATE_ThreeParams(FAndroidDataReceivedIntArrayMC, int32, DataId, FString, Key, const TArray<int32>&, Data);
+
+DECLARE_DYNAMIC_MULTICAST_DELEGATE_ThreeParams(FAndroidDataReceivedShortMC, int32, DataId, FString, Key, int16, Data);
+DECLARE_DYNAMIC_MULTICAST_DELEGATE_ThreeParams(FAndroidDataReceivedShortArrayMC, int32, DataId, FString, Key, const TArray<int16>&, Data);
+
+DECLARE_DYNAMIC_MULTICAST_DELEGATE_ThreeParams(FAndroidDataReceivedLongMC, int32, DataId, FString, Key, int64, Data);
+DECLARE_DYNAMIC_MULTICAST_DELEGATE_ThreeParams(FAndroidDataReceivedLongArrayMC, int32, DataId, FString, Key, const TArray<int64>&, Data);
+
+DECLARE_DYNAMIC_MULTICAST_DELEGATE_ThreeParams(FAndroidDataReceivedFloatMC, int32, DataId, FString, Key, float, Data);
+DECLARE_DYNAMIC_MULTICAST_DELEGATE_ThreeParams(FAndroidDataReceivedFloatArrayMC, int32, DataId, FString, Key, const TArray<float>&, Data);
+
+DECLARE_DYNAMIC_MULTICAST_DELEGATE_ThreeParams(FAndroidDataReceivedDoubleMC, int32, DataId, FString, Key, double, Data);
+DECLARE_DYNAMIC_MULTICAST_DELEGATE_ThreeParams(FAndroidDataReceivedDoubleArrayMC, int32, DataId, FString, Key, const TArray<double>&, Data);
+
+UENUM()
+namespace EASISActiveType
+{
+	enum Type : int
+	{
+		None = 0 UMETA(DisplayName = "None"),
+		LocalOnly = 1 UMETA(DisplayName = "Local only"),
+	};
+}
+
+UCLASS()
+class UAndroidSingleInstanceServiceBPLibrary : public UBlueprintFunctionLibrary
+{
+	GENERATED_UCLASS_BODY()
+public:
+	/** Request startup of Android FileServer */
+	UFUNCTION(BlueprintCallable, Category = "AndroidSingleInstanceService")
+	static bool StartSingleInstanceService();
+
+	/** Request termination of Android FileServer */
+	UFUNCTION(BlueprintCallable, Category = "AndroidSingleInstanceService")
+	static bool StopSingleInstanceService();
+
+	/** Check if Android FileServer is running */
+	UFUNCTION(BlueprintCallable, Category = "AndroidSingleInstanceService")
+	static TEnumAsByte<EASISActiveType::Type> IsFileServerRunning();
+
+	/** Get application name from package */
+	UFUNCTION(BlueprintCallable, Category = "AndroidUtils")
+	static FString GetApplicationNameFromPackage(const FString& PackageName);
+
+	/** Sends a string to Java */
+	UFUNCTION(BlueprintCallable, Category = "AndroidUtils")
+	static bool SendString(const FString& Value);
+
+	/** Returns if startup movie is playing */
+	UFUNCTION(BlueprintPure, Category = "AndroidUtils")
+	static bool IsStartupMoviePlaying();
+
+	/** Stops startup movie if active either paused or playing over duration of fade out */
+	UFUNCTION(BlueprintCallable, Category = "AndroidUtils")
+	static void StopStartupMovie(bool Pause, int Duration);
+
+	/** Set alpha on a view by id over some duration in Java */
+	UFUNCTION(BlueprintCallable, Category = "AndroidUtils")
+	static bool AnimateAlpha(const FString& Key, float TargetAlpha, int Duration);
+
+	/** Set visibiity of a view by id in Java */
+	UFUNCTION(BlueprintCallable, Category = "AndroidUtils")
+	static bool SetVisibility(const FString& Key, bool Visible);
+
+	/** Enables an ImageButton or Button TextView in Java */
+	UFUNCTION(BlueprintCallable, Category = "AndroidUtils")
+	static bool EnableButton(const FString& Key, bool Enabled);
+
+	/** Sets the text on a TextView in Java */
+	UFUNCTION(BlueprintCallable, Category = "AndroidUtils")
+	static bool SendText(const FString& Key, const FString& Value);
+
+	/** Get the UV for corners of the surfaceview */
+	UFUNCTION(BlueprintCallable, Category="AndroidUtils")
+	static FBox2D GetSurfaceViewUVs();
+
+	/** Sends event with parameters to engine callback */
+	UFUNCTION(BlueprintCallable, Category = "AndroidUtils")
+	static void SendAndroidEvent(int32 EventId, const FString& Param1, int32 Param2, int32 Param3, float Param4);
+
+	DECLARE_DYNAMIC_DELEGATE_OneParam(FOnButtonPressed, FString, Name);
+
+	/** Register a delegate to call to receive button presses from Java */
+	UFUNCTION(BlueprintCallable, Category = "AndroidUtils")
+	static void RegisterButtonPressDelegate(FOnButtonPressed Delegate);
+
+	DECLARE_DYNAMIC_DELEGATE_FiveParams(FOnAndroidEventReceived, int32, EventId, FString, Param1, int32, Param2, int32, Param3, float, Param4);
+
+	/** Register a delegate to call to receive events from Java */
+	UFUNCTION(BlueprintCallable, Category = "AndroidUtils")
+	static void RegisterAndroidEventReceivedDelegate(FOnAndroidEventReceived Delegate);
+
+	DECLARE_DYNAMIC_DELEGATE_ThreeParams(FOnAndroidDataReceivedString, int32, DataId, FString, Key, FString, Data);
+	DECLARE_DYNAMIC_DELEGATE_ThreeParams(FOnAndroidDataReceivedStringArray, int32, DataId, FString, Key, const TArray<FString>&, Data);
+
+	/** Register a delegate to call to receive string data from Java */
+	UFUNCTION(BlueprintCallable, Category = "AndroidUtils")
+	static void RegisterAndroidDataReceivedStringDelegate(FOnAndroidDataReceivedString Delegate);
+
+	/** Register a delegate to call to receive string array data from Java */
+	UFUNCTION(BlueprintCallable, Category = "AndroidUtils")
+	static void RegisterAndroidDataReceivedStringArrayDelegate(FOnAndroidDataReceivedStringArray Delegate);
+
+	DECLARE_DYNAMIC_DELEGATE_ThreeParams(FOnAndroidDataReceivedBoolean, int32, DataId, FString, Key, bool, Data);
+	DECLARE_DYNAMIC_DELEGATE_ThreeParams(FOnAndroidDataReceivedBooleanArray, int32, DataId, FString, Key, const TArray<bool>&, Data);
+
+	/** Register a delegate to call to receive boolean data from Java */
+	UFUNCTION(BlueprintCallable, Category = "AndroidUtils")
+	static void RegisterAndroidDataReceivedBooleanDelegate(FOnAndroidDataReceivedBoolean Delegate);
+
+	/** Register a delegate to call to receive boolean array data from Java */
+	UFUNCTION(BlueprintCallable, Category = "AndroidUtils")
+	static void RegisterAndroidDataReceivedBooleanArrayDelegate(FOnAndroidDataReceivedBooleanArray Delegate);
+
+	DECLARE_DYNAMIC_DELEGATE_ThreeParams(FOnAndroidDataReceivedByte, int32, DataId, FString, Key, int8, Data);
+	DECLARE_DYNAMIC_DELEGATE_ThreeParams(FOnAndroidDataReceivedByteArray, int32, DataId, FString, Key, const TArray<int8>&, Data);
+
+	/** Register a delegate to call to receive byte data from Java */
+	UFUNCTION(BlueprintCallable, Category = "AndroidUtils")
+	static void RegisterAndroidDataReceivedByteDelegate(FOnAndroidDataReceivedByte Delegate);
+
+	/** Register a delegate to call to receive byte array data from Java */
+	UFUNCTION(BlueprintCallable, Category = "AndroidUtils")
+	static void RegisterAndroidDataReceivedByteArrayDelegate(FOnAndroidDataReceivedByteArray Delegate);
+
+	DECLARE_DYNAMIC_DELEGATE_ThreeParams(FOnAndroidDataReceivedInt, int32, DataId, FString, Key, int32, Data);
+	DECLARE_DYNAMIC_DELEGATE_ThreeParams(FOnAndroidDataReceivedIntArray, int32, DataId, FString, Key, const TArray<int32>&, Data);
+
+	/** Register a delegate to call to receive int data from Java */
+	UFUNCTION(BlueprintCallable, Category = "AndroidUtils")
+	static void RegisterAndroidDataReceivedIntDelegate(FOnAndroidDataReceivedInt Delegate);
+
+	/** Register a delegate to call to receive int array data from Java */
+	UFUNCTION(BlueprintCallable, Category = "AndroidUtils")
+	static void RegisterAndroidDataReceivedIntArrayDelegate(FOnAndroidDataReceivedIntArray Delegate);
+
+	DECLARE_DYNAMIC_DELEGATE_ThreeParams(FOnAndroidDataReceivedShort, int32, DataId, FString, Key, int16, Data);
+	DECLARE_DYNAMIC_DELEGATE_ThreeParams(FOnAndroidDataReceivedShortArray, int32, DataId, FString, Key, const TArray<int16>&, Data);
+
+	/** Register a delegate to call to receive short data from Java */
+	UFUNCTION(BlueprintCallable, Category = "AndroidUtils")
+	static void RegisterAndroidDataReceivedShortDelegate(FOnAndroidDataReceivedShort Delegate);
+
+	/** Register a delegate to call to receive short array data from Java */
+	UFUNCTION(BlueprintCallable, Category = "AndroidUtils")
+	static void RegisterAndroidDataReceivedShortArrayDelegate(FOnAndroidDataReceivedShortArray Delegate);
+
+	DECLARE_DYNAMIC_DELEGATE_ThreeParams(FOnAndroidDataReceivedLong, int32, DataId, FString, Key, int64, Data);
+	DECLARE_DYNAMIC_DELEGATE_ThreeParams(FOnAndroidDataReceivedLongArray, int32, DataId, FString, Key, const TArray<int64>&, Data);
+
+	/** Register a delegate to call to receive long data from Java */
+	UFUNCTION(BlueprintCallable, Category = "AndroidUtils")
+	static void RegisterAndroidDataReceivedLongDelegate(FOnAndroidDataReceivedLong Delegate);
+
+	/** Register a delegate to call to receive long array data from Java */
+	UFUNCTION(BlueprintCallable, Category = "AndroidUtils")
+	static void RegisterAndroidDataReceivedLongArrayDelegate(FOnAndroidDataReceivedLongArray Delegate);
+
+	DECLARE_DYNAMIC_DELEGATE_ThreeParams(FOnAndroidDataReceivedFloat, int32, DataId, FString, Key, float, Data);
+	DECLARE_DYNAMIC_DELEGATE_ThreeParams(FOnAndroidDataReceivedFloatArray, int32, DataId, FString, Key, const TArray<float>&, Data);
+
+	/** Register a delegate to call to receive float data from Java */
+	UFUNCTION(BlueprintCallable, Category = "AndroidUtils")
+	static void RegisterAndroidDataReceivedFloatDelegate(FOnAndroidDataReceivedFloat Delegate);
+
+	/** Register a delegate to call to receive float array data from Java */
+	UFUNCTION(BlueprintCallable, Category = "AndroidUtils")
+	static void RegisterAndroidDataReceivedFloatArrayDelegate(FOnAndroidDataReceivedFloatArray Delegate);
+
+	DECLARE_DYNAMIC_DELEGATE_ThreeParams(FOnAndroidDataReceivedDouble, int32, DataId, FString, Key, double, Data);
+	DECLARE_DYNAMIC_DELEGATE_ThreeParams(FOnAndroidDataReceivedDoubleArray, int32, DataId, FString, Key, const TArray<double>&, Data);
+
+	/** Register a delegate to call to receive double data from Java */
+	UFUNCTION(BlueprintCallable, Category = "AndroidUtils")
+	static void RegisterAndroidDataReceivedDoubleDelegate(FOnAndroidDataReceivedDouble Delegate);
+
+	/** Register a delegate to call to receive double array data from Java */
+	UFUNCTION(BlueprintCallable, Category = "AndroidUtils")
+	static void RegisterAndroidDataReceivedDoubleArrayDelegate(FOnAndroidDataReceivedDoubleArray Delegate);
+};
+
+#if UE_ENABLE_INCLUDE_ORDER_DEPRECATED_IN_5_2
+#include "CoreMinimal.h"
+#endif
diff --git a/Engine/Plugins/Runtime/AndroidSingleInstanceService/Source/AndroidSingleInstanceService/Private/AndroidSingleInstanceService.cpp b/Engine/Plugins/Runtime/AndroidSingleInstanceService/Source/AndroidSingleInstanceService/Private/AndroidSingleInstanceService.cpp
new file mode 100644
--- /dev/null
+++ b/Engine/Plugins/Runtime/AndroidSingleInstanceService/Source/AndroidSingleInstanceService/Private/AndroidSingleInstanceService.cpp
@@ -0,0 +1,23 @@
+// Copyright Epic Games, Inc. All Rights Reserved.
+
+#include "AndroidSingleInstanceService.h"
+#include "Modules/ModuleManager.h"
+
+#define LOCTEXT_NAMESPACE "FAndroidSingleInstanceServiceModule"
+
+void FAndroidSingleInstanceServiceModule::StartupModule()
+{
+	// This code will execute after your module is loaded into memory; the exact timing is specified in the .uplugin file per-module
+	
+}
+
+void FAndroidSingleInstanceServiceModule::ShutdownModule()
+{
+	// This function may be called during shutdown to clean up your module.  For modules that support dynamic reloading,
+	// we call this function before unloading the module.
+	
+}
+
+#undef LOCTEXT_NAMESPACE
+	
+IMPLEMENT_MODULE(FAndroidSingleInstanceServiceModule, AndroidSingleInstanceService)
diff --git a/Engine/Plugins/Runtime/AndroidSingleInstanceService/Source/AndroidSingleInstanceService/Private/AndroidSingleInstanceServiceBPLibrary.cpp b/Engine/Plugins/Runtime/AndroidSingleInstanceService/Source/AndroidSingleInstanceService/Private/AndroidSingleInstanceServiceBPLibrary.cpp
new file mode 100644
--- /dev/null
+++ b/Engine/Plugins/Runtime/AndroidSingleInstanceService/Source/AndroidSingleInstanceService/Private/AndroidSingleInstanceServiceBPLibrary.cpp
@@ -0,0 +1,516 @@
+// Copyright Epic Games, Inc. All Rights Reserved.
+
+#include "AndroidSingleInstanceServiceBPLibrary.h"
+
+#include UE_INLINE_GENERATED_CPP_BY_NAME(AndroidSingleInstanceServiceBPLibrary)
+
+FButtonPressedMC GlobalOnButtonPressed;
+FAndroidEventReceivedMC GlobalOnAndroidEventReceived;
+
+void UAndroidSingleInstanceServiceBPLibrary::RegisterButtonPressDelegate(FOnButtonPressed Delegate)
+{
+	GlobalOnButtonPressed.Add(Delegate);
+}
+
+void UAndroidSingleInstanceServiceBPLibrary::RegisterAndroidEventReceivedDelegate(FOnAndroidEventReceived Delegate)
+{
+	GlobalOnAndroidEventReceived.Add(Delegate);
+}
+
+#define REGISTER_DATA_RECEIVED_DELEGATE(_TYPE_)	\
+FAndroidDataReceived ## _TYPE_ ## MC GlobalOnAndroidDataReceived ## _TYPE_; \
+void UAndroidSingleInstanceServiceBPLibrary::RegisterAndroidDataReceived ## _TYPE_ ## Delegate(FOnAndroidDataReceived ## _TYPE_ Delegate) \
+{ \
+	GlobalOnAndroidDataReceived ## _TYPE_.Add(Delegate); \
+}
+
+REGISTER_DATA_RECEIVED_DELEGATE(String);
+REGISTER_DATA_RECEIVED_DELEGATE(StringArray);
+REGISTER_DATA_RECEIVED_DELEGATE(Boolean);
+REGISTER_DATA_RECEIVED_DELEGATE(BooleanArray);
+REGISTER_DATA_RECEIVED_DELEGATE(Byte);
+REGISTER_DATA_RECEIVED_DELEGATE(ByteArray);
+REGISTER_DATA_RECEIVED_DELEGATE(Int);
+REGISTER_DATA_RECEIVED_DELEGATE(IntArray);
+REGISTER_DATA_RECEIVED_DELEGATE(Short);
+REGISTER_DATA_RECEIVED_DELEGATE(ShortArray);
+REGISTER_DATA_RECEIVED_DELEGATE(Long);
+REGISTER_DATA_RECEIVED_DELEGATE(LongArray);
+REGISTER_DATA_RECEIVED_DELEGATE(Float);
+REGISTER_DATA_RECEIVED_DELEGATE(FloatArray);
+REGISTER_DATA_RECEIVED_DELEGATE(Double);
+REGISTER_DATA_RECEIVED_DELEGATE(DoubleArray);
+
+#if PLATFORM_ANDROID
+
+#include "Engine/Engine.h"
+#include "Android/AndroidJNI.h"
+#include "Android/AndroidApplication.h"
+#include "Android/AndroidWindow.h"
+
+#if USE_ANDROID_JNI
+#include <jni.h>
+#endif
+#endif
+
+bool UAndroidSingleInstanceServiceBPLibrary::StartSingleInstanceService()
+{
+	bool result = false;
+#if PLATFORM_ANDROID
+	if (JNIEnv* Env = FAndroidApplication::GetJavaEnv())
+	{
+		static jmethodID StartFunc = FJavaWrapper::FindMethod(Env, FJavaWrapper::GameActivityClassID, "AndroidThunkJava_AndroidSingleInstanceService_Start", "()Z", false);
+		if (StartFunc != nullptr)
+		{
+			result = FJavaWrapper::CallBooleanMethod(Env, FJavaWrapper::GameActivityThis, StartFunc);
+		}
+	}
+#endif
+	return result;
+}
+
+bool UAndroidSingleInstanceServiceBPLibrary::StopSingleInstanceService()
+{
+	bool result = false;
+#if PLATFORM_ANDROID
+	if (JNIEnv* Env = FAndroidApplication::GetJavaEnv())
+	{
+		static jmethodID StopFunc = FJavaWrapper::FindMethod(Env, FJavaWrapper::GameActivityClassID, "AndroidThunkJava_AndroidSingleInstanceService_Stop", "()Z", false);
+		if (StopFunc != nullptr)
+		{
+			result = FJavaWrapper::CallBooleanMethod(Env, FJavaWrapper::GameActivityThis, StopFunc);
+		}
+	}
+#endif
+	return result;
+}
+
+TEnumAsByte<EASISActiveType::Type> UAndroidSingleInstanceServiceBPLibrary::IsFileServerRunning()
+{
+	TEnumAsByte<EASISActiveType::Type> result = EASISActiveType::None;
+#if PLATFORM_ANDROID
+	if (JNIEnv* Env = FAndroidApplication::GetJavaEnv())
+	{
+		static jmethodID IsRunningFunc = FJavaWrapper::FindMethod(Env, FJavaWrapper::GameActivityClassID, "AndroidThunkJava_AndroidSingleInstanceService_IsRunning", "()I", false);
+		if (IsRunningFunc != nullptr)
+		{
+			result = (TEnumAsByte<EASISActiveType::Type>)FJavaWrapper::CallIntMethod(Env, FJavaWrapper::GameActivityThis, IsRunningFunc);
+		}
+	}
+#endif
+	return result;
+}
+
+#if PLATFORM_ANDROID
+
+#ifndef USE_ANDROID_MIXEDSAMPLE
+#define USE_ANDROID_MIXEDSAMPLE 0
+#endif
+
+UAndroidSingleInstanceServiceBPLibrary::UAndroidSingleInstanceServiceBPLibrary(const FObjectInitializer& ObjectInitializer)
+	: Super(ObjectInitializer)
+{
+}
+
+FString UAndroidSingleInstanceServiceBPLibrary::GetApplicationNameFromPackage(const FString& PackageName)
+{
+	FString Result = FString("");
+#if USE_ANDROID_MIXEDSAMPLE
+	if (JNIEnv* Env = FAndroidApplication::GetJavaEnv())
+	{
+		static jmethodID JavaFunc = FJavaWrapper::FindMethod(Env, FJavaWrapper::GameActivityClassID, "AndroidThunkJava_MixedSample_GetApplicationNameFromPackage", "(Ljava/lang/String;)Ljava/lang/String;", false);
+		if (JavaFunc != nullptr)
+		{
+			auto Argument = FJavaHelper::ToJavaString(Env, PackageName);
+			Result = FJavaHelper::FStringFromLocalRef(Env, (jstring)FJavaWrapper::CallObjectMethod(Env, FJavaWrapper::GameActivityThis, JavaFunc, *Argument));
+		}
+	}
+#endif
+	return Result;
+}
+
+bool UAndroidSingleInstanceServiceBPLibrary::SendString(const FString& Value)
+{
+	bool bResult = false;
+#if USE_ANDROID_MIXEDSAMPLE
+	if (JNIEnv* Env = FAndroidApplication::GetJavaEnv())
+	{
+		static jmethodID JavaFunc = FJavaWrapper::FindMethod(Env, FJavaWrapper::GameActivityClassID, "AndroidThunkJava_MixedSample_ReceiveString", "(Ljava/lang/String;)Z", false);
+		if (JavaFunc != nullptr)
+		{
+			auto Argument = FJavaHelper::ToJavaString(Env, Value);
+			bResult = FJavaWrapper::CallBooleanMethod(Env, FJavaWrapper::GameActivityThis, JavaFunc, *Argument);
+		}
+	}
+#endif
+	return bResult;
+}
+
+bool UAndroidSingleInstanceServiceBPLibrary::SendText(const FString& Key, const FString& Value)
+{
+	bool bResult = false;
+#if USE_ANDROID_MIXEDSAMPLE
+	if (JNIEnv* Env = FAndroidApplication::GetJavaEnv())
+	{
+		static jmethodID JavaFunc = FJavaWrapper::FindMethod(Env, FJavaWrapper::GameActivityClassID, "AndroidThunkJava_MixedSample_SetText", "(Ljava/lang/String;Ljava/lang/String;)Z", false);
+		if (JavaFunc != nullptr)
+		{
+			auto Argument1 = FJavaHelper::ToJavaString(Env, Key);
+			auto Argument2 = FJavaHelper::ToJavaString(Env, Value);
+			bResult = FJavaWrapper::CallBooleanMethod(Env, FJavaWrapper::GameActivityThis, JavaFunc, *Argument1, *Argument2);
+		}
+	}
+#endif
+	return bResult;
+}
+
+FBox2D UAndroidSingleInstanceServiceBPLibrary::GetSurfaceViewUVs()
+{
+	float RealScreenWidth = 1024.0f;
+	float RealScreenHeight = 800.0f;
+
+	if (FString* ConfigRulesWidth = FAndroidMisc::GetConfigRulesVariable(TEXT("screenWidth")))
+	{
+		RealScreenWidth = FCString::Atof(**ConfigRulesWidth);
+	}
+	if (FString* ConfigRulesHeight = FAndroidMisc::GetConfigRulesVariable(TEXT("screenHeight")))
+	{
+		RealScreenHeight = FCString::Atof(**ConfigRulesHeight);
+	}
+
+	FPlatformRect Rect = FAndroidWindow::GetScreenRect();
+
+	STANDALONE_DEBUG_LOGf(LogAndroid, TEXT("** GetSurfaceViewUVs: Real WxH: %f,%f,  Bounds: %d,%d - %d,%d"), RealScreenWidth, RealScreenHeight, Rect.Left, Rect.Top, Rect.Right, Rect.Bottom);
+
+	FVector2D UVMin = FVector2D(Rect.Left / RealScreenWidth, Rect.Top / RealScreenHeight);
+	FVector2D UVMax = FVector2D(Rect.Right / RealScreenWidth, Rect.Bottom / RealScreenHeight);
+
+	return FBox2D(UVMin, UVMax);
+}
+
+bool UAndroidSingleInstanceServiceBPLibrary::IsStartupMoviePlaying()
+{
+	bool bResult = false;
+#if USE_ANDROID_MIXEDSAMPLE
+	if (JNIEnv* Env = FAndroidApplication::GetJavaEnv())
+	{
+		static jmethodID JavaFunc = FJavaWrapper::FindMethod(Env, FJavaWrapper::GameActivityClassID, "AndroidThunkJava_MixedSample_IsStartupMoviePlaying", "()Z", false);
+		if (JavaFunc != nullptr)
+		{
+			bResult = FJavaWrapper::CallBooleanMethod(Env, FJavaWrapper::GameActivityThis, JavaFunc);
+		}
+	}
+#endif
+	return bResult;
+}
+
+void UAndroidSingleInstanceServiceBPLibrary::StopStartupMovie(bool Paused, int Duration)
+{
+#if USE_ANDROID_MIXEDSAMPLE
+	if (JNIEnv* Env = FAndroidApplication::GetJavaEnv())
+	{
+		static jmethodID JavaFunc = FJavaWrapper::FindMethod(Env, FJavaWrapper::GameActivityClassID, "AndroidThunkJava_MixedSample_StopStartupMovie", "(ZI)V", false);
+		if (JavaFunc != nullptr)
+		{
+			FJavaWrapper::CallVoidMethod(Env, FJavaWrapper::GameActivityThis, JavaFunc, Paused, Duration);
+		}
+	}
+#endif
+}
+
+bool UAndroidSingleInstanceServiceBPLibrary::AnimateAlpha(const FString& Key, float TargetAlpha, int Duration)
+{
+	bool bResult = false;
+#if USE_ANDROID_MIXEDSAMPLE
+	if (JNIEnv* Env = FAndroidApplication::GetJavaEnv())
+	{
+		static jmethodID JavaFunc = FJavaWrapper::FindMethod(Env, FJavaWrapper::GameActivityClassID, "AndroidThunkJava_MixedSample_AnimateAlpha", "(Ljava/lang/String;FI)Z", false);
+		if (JavaFunc != nullptr)
+		{
+			auto Argument = FJavaHelper::ToJavaString(Env, Key);
+			bResult = FJavaWrapper::CallBooleanMethod(Env, FJavaWrapper::GameActivityThis, JavaFunc, *Argument, TargetAlpha, Duration);
+		}
+	}
+#endif
+	return bResult;
+}
+
+bool UAndroidSingleInstanceServiceBPLibrary::SetVisibility(const FString& Key, bool Visible)
+{
+	bool bResult = false;
+#if USE_ANDROID_MIXEDSAMPLE
+	if (JNIEnv* Env = FAndroidApplication::GetJavaEnv())
+	{
+		static jmethodID JavaFunc = FJavaWrapper::FindMethod(Env, FJavaWrapper::GameActivityClassID, "AndroidThunkJava_MixedSample_SetVisibility", "(Ljava/lang/String;Z)Z", false);
+		if (JavaFunc != nullptr)
+		{
+			auto Argument = FJavaHelper::ToJavaString(Env, Key);
+			bResult = FJavaWrapper::CallBooleanMethod(Env, FJavaWrapper::GameActivityThis, JavaFunc, *Argument, Visible);
+		}
+	}
+#endif
+	return bResult;
+}
+
+bool UAndroidSingleInstanceServiceBPLibrary::EnableButton(const FString& Key, bool Enabled)
+{
+	bool bResult = false;
+#if USE_ANDROID_MIXEDSAMPLE
+	if (JNIEnv* Env = FAndroidApplication::GetJavaEnv())
+	{
+		static jmethodID JavaFunc = FJavaWrapper::FindMethod(Env, FJavaWrapper::GameActivityClassID, "AndroidThunkJava_MixedSample_EnableButton", "(Ljava/lang/String;Z)Z", false);
+		if (JavaFunc != nullptr)
+		{
+			auto Argument = FJavaHelper::ToJavaString(Env, Key);
+			bResult = FJavaWrapper::CallBooleanMethod(Env, FJavaWrapper::GameActivityThis, JavaFunc, *Argument, Enabled);
+		}
+	}
+#endif
+	return bResult;
+}
+
+#define EVENTTYPE_INIT						0
+#define EVENTTYPE_POST_ENGINE_INIT			1
+#define EVENTTYPE_ENGINELOOP_INIT_COMPLETE	2
+#define EVENTTYPE_FRAME_BEGIN				3
+#define EVENTTYPE_FRAME_END					4
+#define EVENTTYPE_PRE_LOAD_MAP				5
+#define EVENTTYPE_POST_LOAD_MAP				6
+#define EVENTTYPE_ACTION					7
+
+void UAndroidSingleInstanceServiceBPLibrary::SendAndroidEvent(int32 EventId, const FString& Param1, int32 Param2, int32 Param3, float Param4)
+{
+#if USE_ANDROID_STANDALONE
+	extern void AndroidThunkCpp_Engine_SendEvent(int32 event, const FString& param1, int32 param2, int32 param3, float param4);
+	AndroidThunkCpp_Engine_SendEvent(EVENTTYPE_ACTION, Param1, Param2, Param3, Param4);
+#endif
+}
+
+// receives calls from Java for events
+JNI_METHOD void Java_com_epicgames_makeaar_Engine_nativeEventReceived(JNIEnv* jenv, jobject thiz, jint eventId, jstring param1, jint param2, jint param3, jfloat param4)
+{
+	FString ParamString = FJavaHelper::FStringFromParam(jenv, param1);
+
+	if (FTaskGraphInterface::IsRunning())
+	{
+		STANDALONE_DEBUG_LOGf(LogAndroid,TEXT("nativeEventReceived: %d, %s, %d, %d, %f"), eventId, *ParamString, param2, param3, param4);
+		FSimpleDelegateGraphTask::CreateAndDispatchWhenReady(
+			FSimpleDelegateGraphTask::FDelegate::CreateLambda([=]()
+		{
+			// trigger blueprint event
+			GlobalOnAndroidEventReceived.Broadcast(eventId, ParamString, param2, param3, param4);
+		}), TStatId(), nullptr, ENamedThreads::GameThread);
+	}
+	else
+	{
+		STANDALONE_DEBUG_LOGf(LogAndroid,TEXT("nativeEventReceived: %d, %s, %d, %d, %f - IGNORED (engine not ready)"), eventId, *ParamString, param2, param3, param4);
+	}
+}
+
+// receives calls from MixedSampleActivity for pressed button
+JNI_METHOD int Java_com_epicgames_makeaar_MixedSampleActivity_nativePressedButton(JNIEnv* jenv, jobject thiz, jstring name)
+{
+	FString ButtonName = FJavaHelper::FStringFromParam(jenv, name);
+
+	STANDALONE_DEBUG_LOGf(LogAndroid,TEXT("nativePressedButton: %s"), *ButtonName);
+
+	if (FTaskGraphInterface::IsRunning())
+	{
+		FSimpleDelegateGraphTask::CreateAndDispatchWhenReady(
+			FSimpleDelegateGraphTask::FDelegate::CreateLambda([=]()
+		{
+			// trigger blueprint event
+			GlobalOnButtonPressed.Broadcast(ButtonName);
+		}), TStatId(), nullptr, ENamedThreads::GameThread);
+		return 1;
+	}
+
+	return 0;
+}
+
+// receives calls from Java for data
+JNI_METHOD void Java_com_epicgames_makeaar_Engine_nativeDataReceivedString(JNIEnv* jenv, jobject thiz, jint dataId, jstring key, jstring value)
+{
+	FString KeyString = FJavaHelper::FStringFromParam(jenv, key);
+	FString Value = FJavaHelper::FStringFromParam(jenv, value);
+
+	if (FTaskGraphInterface::IsRunning())
+	{
+		STANDALONE_DEBUG_LOGf(LogAndroid,TEXT("nativeEventDataReceivedString: %d, %s, %s"), dataId, *KeyString, *Value);
+		FSimpleDelegateGraphTask::CreateAndDispatchWhenReady(
+			FSimpleDelegateGraphTask::FDelegate::CreateLambda([=]()
+		{
+			// trigger blueprint event
+			GlobalOnAndroidDataReceivedString.Broadcast(dataId, KeyString, Value);
+		}), TStatId(), nullptr, ENamedThreads::GameThread);
+	}
+	else
+	{
+		STANDALONE_DEBUG_LOGf(LogAndroid,TEXT("nativeDataReceivedString: %d, %s, %s - IGNORED (engine not ready)"), dataId, *KeyString, *Value);
+	}
+}
+
+JNI_METHOD void Java_com_epicgames_makeaar_Engine_nativeDataReceivedStringArray(JNIEnv* jenv, jobject thiz, jint dataId, jstring key, jobjectArray value)
+{
+	FString KeyString = FJavaHelper::FStringFromParam(jenv, key);
+	TArray<FString> Values;
+	if (value != nullptr)
+	{
+		int32 Count = jenv->GetArrayLength(value);
+		for (int32 Index=0; Index < Count; Index++)
+		{
+			Values.Add(FJavaHelper::FStringFromLocalRef(jenv, (jstring)(jenv->GetObjectArrayElement(value, Index))));
+		}
+	}
+
+	if (FTaskGraphInterface::IsRunning())
+	{
+		STANDALONE_DEBUG_LOGf(LogAndroid,TEXT("nativeEventDataReceivedStringArray: %d, %s"), dataId, *KeyString);
+		FSimpleDelegateGraphTask::CreateAndDispatchWhenReady(
+			FSimpleDelegateGraphTask::FDelegate::CreateLambda([=]()
+		{
+			// trigger blueprint event
+			GlobalOnAndroidDataReceivedStringArray.Broadcast(dataId, KeyString, Values);
+		}), TStatId(), nullptr, ENamedThreads::GameThread);
+	}
+	else
+	{
+		STANDALONE_DEBUG_LOGf(LogAndroid,TEXT("nativeDataReceivedStringArray: %d, %s - IGNORED (engine not ready)"), dataId, *KeyString);
+	}
+}
+
+#define NAUTIL_STRING(_IN) #_IN
+
+#define NATIVE_DATA_RECEIVED_VALUE(_NAMETYPE_, _JAVATYPE_, _CPPTYPE_) \
+JNI_METHOD void Java_com_epicgames_makeaar_Engine_nativeDataReceived ## _NAMETYPE_ (JNIEnv* jenv, jobject thiz, jint dataId, jstring key, _JAVATYPE_ value) \
+{ \
+	FString KeyString = FJavaHelper::FStringFromParam(jenv, key); \
+	_CPPTYPE_ Value = (_CPPTYPE_)value; \
+	\
+	if (FTaskGraphInterface::IsRunning()) \
+	{ \
+		STANDALONE_DEBUG_LOGf(LogAndroid,TEXT("nativeEventReceived" NAUTIL_STRING(_NAMETYPE_) ": %d, %s"), dataId, *KeyString); \
+		FSimpleDelegateGraphTask::CreateAndDispatchWhenReady( \
+			FSimpleDelegateGraphTask::FDelegate::CreateLambda([=]() \
+		{ \
+			GlobalOnAndroidDataReceived ## _NAMETYPE_.Broadcast(dataId, KeyString, Value); \
+		}), TStatId(), nullptr, ENamedThreads::GameThread); \
+	} \
+	else \
+	{ \
+		STANDALONE_DEBUG_LOGf(LogAndroid,TEXT("nativeDataReceived" NAUTIL_STRING(_NAMETYPE_) ": %d, %s - IGNORED (engine not ready)"), dataId, *KeyString); \
+	} \
+}
+
+NATIVE_DATA_RECEIVED_VALUE(Boolean, jboolean, bool);
+NATIVE_DATA_RECEIVED_VALUE(Byte, jbyte, int8);
+NATIVE_DATA_RECEIVED_VALUE(Int, jint, int32);
+NATIVE_DATA_RECEIVED_VALUE(Short, jshort, int16);
+NATIVE_DATA_RECEIVED_VALUE(Long, jlong, int64);
+NATIVE_DATA_RECEIVED_VALUE(Float, jfloat, float);
+NATIVE_DATA_RECEIVED_VALUE(Double, jdouble, double);
+
+#define NATIVE_DATA_RECEIVED_VALUEARRAY(_NAMETYPE_, _JAVATYPE_, _CPPTYPE_) \
+JNI_METHOD void Java_com_epicgames_makeaar_Engine_nativeDataReceived ## _NAMETYPE_ ## Array(JNIEnv* jenv, jobject thiz, jint dataId, jstring key, _JAVATYPE_ ## Array value) \
+{ \
+	FString KeyString = FJavaHelper::FStringFromParam(jenv, key); \
+	TArray<_CPPTYPE_> Values; \
+	if (value != nullptr) \
+	{ \
+		int32 Count = jenv->GetArrayLength(value); \
+		_JAVATYPE_ *values = jenv->Get ## _NAMETYPE_ ## ArrayElements(value, 0); \
+		for (int32 Index = 0; Index < Count; Index++) \
+		{ \
+			Values.Add((_CPPTYPE_)values[Index]); \
+		} \
+		jenv->Release ## _NAMETYPE_ ## ArrayElements(value, values, 0); \
+	} \
+	\
+	if (FTaskGraphInterface::IsRunning()) \
+	{ \
+		STANDALONE_DEBUG_LOGf(LogAndroid,TEXT("nativeEventReceived" NAUTIL_STRING(_NAMETYPE_) "Array: %d, %s"), dataId, *KeyString); \
+		FSimpleDelegateGraphTask::CreateAndDispatchWhenReady( \
+			FSimpleDelegateGraphTask::FDelegate::CreateLambda([=]() \
+		{ \
+			GlobalOnAndroidDataReceived ## _NAMETYPE_ ## Array.Broadcast(dataId, KeyString, Values); \
+		}), TStatId(), nullptr, ENamedThreads::GameThread); \
+	} \
+	else \
+	{ \
+		STANDALONE_DEBUG_LOGf(LogAndroid,TEXT("nativeDataReceived" NAUTIL_STRING(_NAMETYPE_) "Array: %d, %s - IGNORED (engine not ready)"), dataId, *KeyString); \
+	} \
+}
+
+NATIVE_DATA_RECEIVED_VALUEARRAY(Boolean, jboolean, bool);
+NATIVE_DATA_RECEIVED_VALUEARRAY(Byte, jbyte, int8);
+NATIVE_DATA_RECEIVED_VALUEARRAY(Int, jint, int32);
+NATIVE_DATA_RECEIVED_VALUEARRAY(Short, jshort, int16);
+NATIVE_DATA_RECEIVED_VALUEARRAY(Long, jlong, int64);
+NATIVE_DATA_RECEIVED_VALUEARRAY(Float, jfloat, float);
+NATIVE_DATA_RECEIVED_VALUEARRAY(Double, jdouble, double);
+
+#else
+
+UAndroidSingleInstanceServiceBPLibrary::UAndroidSingleInstanceServiceBPLibrary(const FObjectInitializer& ObjectInitializer)
+	: Super(ObjectInitializer)
+{
+}
+
+FString UAndroidSingleInstanceServiceBPLibrary::GetApplicationNameFromPackage(const FString& PackageName)
+{
+	FString Result = FString("Unknown");
+	return Result;
+}
+
+bool UAndroidSingleInstanceServiceBPLibrary::SendString(const FString& Value)
+{
+	// do nothing
+	return false;
+}
+
+bool UAndroidSingleInstanceServiceBPLibrary::SendText(const FString& Key, const FString& Value)
+{
+	// do nothing
+	return false;
+}
+
+FBox2D UAndroidSingleInstanceServiceBPLibrary::GetSurfaceViewUVs()
+{
+	return FBox2D(FVector2D(0.0f, 0.0f), FVector2D(1.0f, 1.0f));
+}
+
+bool UAndroidSingleInstanceServiceBPLibrary::IsStartupMoviePlaying()
+{
+	return false;
+}
+
+void UAndroidSingleInstanceServiceBPLibrary::StopStartupMovie(bool Paused, int Duration)
+{
+	// do nothing
+	return;
+}
+
+bool UAndroidSingleInstanceServiceBPLibrary::AnimateAlpha(const FString& Key, float TargetAlpha, int Duration)
+{
+	// do nothing
+	return false;
+}
+
+bool UAndroidSingleInstanceServiceBPLibrary::SetVisibility(const FString& Key, bool Visible)
+{
+	// do nothing
+	return false;
+}
+
+bool UAndroidSingleInstanceServiceBPLibrary::EnableButton(const FString& Key, bool Enabled)
+{
+	// do nothing
+	return false;
+}
+
+void UAndroidSingleInstanceServiceBPLibrary::SendAndroidEvent(int32 EventId, const FString& Param1, int32 Param2, int32 Param3, float Param4)
+{
+}
+
+#endif
+
diff --git a/Engine/Plugins/Runtime/AndroidSingleInstanceService/Source/AndroidSingleInstanceService/Public/AndroidSingleInstanceService.h b/Engine/Plugins/Runtime/AndroidSingleInstanceService/Source/AndroidSingleInstanceService/Public/AndroidSingleInstanceService.h
new file mode 100644
--- /dev/null
+++ b/Engine/Plugins/Runtime/AndroidSingleInstanceService/Source/AndroidSingleInstanceService/Public/AndroidSingleInstanceService.h
@@ -0,0 +1,18 @@
+// Copyright Epic Games, Inc. All Rights Reserved.
+
+#pragma once
+
+#include "Modules/ModuleInterface.h"
+
+class FAndroidSingleInstanceServiceModule : public IModuleInterface
+{
+public:
+
+	/** IModuleInterface implementation */
+	virtual void StartupModule() override;
+	virtual void ShutdownModule() override;
+};
+
+#if UE_ENABLE_INCLUDE_ORDER_DEPRECATED_IN_5_2
+#include "Modules/ModuleManager.h"
+#endif
diff --git a/Engine/Plugins/Runtime/AndroidSingleInstanceService/Source/AndroidSingleInstanceServiceEditor/AndroidSingleInstanceServiceEditor.Build.cs b/Engine/Plugins/Runtime/AndroidSingleInstanceService/Source/AndroidSingleInstanceServiceEditor/AndroidSingleInstanceServiceEditor.Build.cs
new file mode 100644
--- /dev/null
+++ b/Engine/Plugins/Runtime/AndroidSingleInstanceService/Source/AndroidSingleInstanceServiceEditor/AndroidSingleInstanceServiceEditor.Build.cs
@@ -0,0 +1,17 @@
+// Copyright Epic Games, Inc. All Rights Reserved.
+
+namespace UnrealBuildTool.Rules
+{
+	public class AndroidSingleInstanceServiceEditor : ModuleRules
+	{
+		public AndroidSingleInstanceServiceEditor(ReadOnlyTargetRules Target) : base(Target)
+		{
+			PrivateDependencyModuleNames.AddRange(
+				new string[] {
+					"Core",
+					"CoreUObject",
+					"UnrealEd",
+				});
+		}
+	}
+}
diff --git a/Engine/Plugins/Runtime/AndroidSingleInstanceService/Source/AndroidSingleInstanceServiceEditor/Private/AndroidSingleInstanceServiceEditor.cpp b/Engine/Plugins/Runtime/AndroidSingleInstanceService/Source/AndroidSingleInstanceServiceEditor/Private/AndroidSingleInstanceServiceEditor.cpp
new file mode 100644
--- /dev/null
+++ b/Engine/Plugins/Runtime/AndroidSingleInstanceService/Source/AndroidSingleInstanceServiceEditor/Private/AndroidSingleInstanceServiceEditor.cpp
@@ -0,0 +1,43 @@
+// Copyright Epic Games, Inc. All Rights Reserved.
+
+#include "AndroidSingleInstanceServiceEditor.h"
+#include "ISettingsModule.h"
+#include "AndroidSingleInstanceServiceRuntimeSettings.h"
+#include "Modules/ModuleManager.h"
+
+/**
+ * Implements the AndroidSingleInstanceServiceEditor module.
+ */
+
+#define LOCTEXT_NAMESPACE "AndroidSingleInstanceService"
+
+void FAndroidSingleInstanceServiceEditorModule::StartupModule()
+{
+	// register settings
+	ISettingsModule* SettingsModule = FModuleManager::GetModulePtr<ISettingsModule>("Settings");
+
+	if (SettingsModule != nullptr)	
+	{
+		SettingsModule->RegisterSettings("Project", "Plugins", "AndroidSingleInstanceService",
+			LOCTEXT("AndroidSingleInstanceServiceSettingsName", "AndroidSingleInstanceService"),
+			LOCTEXT("AndroidSingleInstanceServiceSettingsDescription", "Project settings for AndroidSingleInstanceService plugin"),
+			GetMutableDefault<UAndroidSingleInstanceServiceRuntimeSettings>()
+		);
+	}
+}
+
+void FAndroidSingleInstanceServiceEditorModule::ShutdownModule()
+{
+	// unregister settings
+	ISettingsModule* SettingsModule = FModuleManager::GetModulePtr<ISettingsModule>("Settings");
+
+	if (SettingsModule != nullptr)
+	{
+       		SettingsModule->UnregisterSettings("Project", "Plugins", "AndroidSingleInstanceService");
+	}
+}
+
+
+IMPLEMENT_MODULE(FAndroidSingleInstanceServiceEditorModule, AndroidSingleInstanceServiceEditor);
+
+#undef LOCTEXT_NAMESPACE
diff --git a/Engine/Plugins/Runtime/AndroidSingleInstanceService/Source/AndroidSingleInstanceServiceEditor/Private/AndroidSingleInstanceServiceRuntimeSettings.cpp b/Engine/Plugins/Runtime/AndroidSingleInstanceService/Source/AndroidSingleInstanceServiceEditor/Private/AndroidSingleInstanceServiceRuntimeSettings.cpp
new file mode 100644
--- /dev/null
+++ b/Engine/Plugins/Runtime/AndroidSingleInstanceService/Source/AndroidSingleInstanceServiceEditor/Private/AndroidSingleInstanceServiceRuntimeSettings.cpp
@@ -0,0 +1,15 @@
+// Copyright Epic Games, Inc. All Rights Reserved.
+
+#include "AndroidSingleInstanceServiceRuntimeSettings.h"
+
+#include UE_INLINE_GENERATED_CPP_BY_NAME(AndroidSingleInstanceServiceRuntimeSettings)
+
+//////////////////////////////////////////////////////////////////////////
+// UAndroidSingleInstanceServiceRuntimeSettings
+
+UAndroidSingleInstanceServiceRuntimeSettings::UAndroidSingleInstanceServiceRuntimeSettings(const FObjectInitializer& ObjectInitializer)
+        : Super(ObjectInitializer)
+        , bEnablePlugin(false)
+		, bCompileASISProject(false)
+{
+}
\ No newline at end of file
diff --git a/Engine/Plugins/Runtime/AndroidSingleInstanceService/Source/AndroidSingleInstanceServiceEditor/Public/AndroidSingleInstanceServiceEditor.h b/Engine/Plugins/Runtime/AndroidSingleInstanceService/Source/AndroidSingleInstanceServiceEditor/Public/AndroidSingleInstanceServiceEditor.h
new file mode 100644
--- /dev/null
+++ b/Engine/Plugins/Runtime/AndroidSingleInstanceService/Source/AndroidSingleInstanceServiceEditor/Public/AndroidSingleInstanceServiceEditor.h
@@ -0,0 +1,22 @@
+// Copyright Epic Games, Inc. All Rights Reserved.
+
+#pragma once
+
+#include "Modules/ModuleInterface.h"
+
+
+class FAndroidSingleInstanceServiceEditorModule
+	: public IModuleInterface
+{
+public:
+
+	//~ IModuleInterface interface
+
+	virtual void StartupModule() override;
+	virtual void ShutdownModule() override;
+};
+
+#if UE_ENABLE_INCLUDE_ORDER_DEPRECATED_IN_5_2
+#include "AndroidSingleInstanceServiceRuntimeSettings.h"
+#include "Modules/ModuleManager.h"
+#endif
diff --git a/Engine/Plugins/Runtime/AndroidSingleInstanceService/Source/AndroidSingleInstanceServiceEditor/Public/AndroidSingleInstanceServiceRuntimeSettings.h b/Engine/Plugins/Runtime/AndroidSingleInstanceService/Source/AndroidSingleInstanceServiceEditor/Public/AndroidSingleInstanceServiceRuntimeSettings.h
new file mode 100644
--- /dev/null
+++ b/Engine/Plugins/Runtime/AndroidSingleInstanceService/Source/AndroidSingleInstanceServiceEditor/Public/AndroidSingleInstanceServiceRuntimeSettings.h
@@ -0,0 +1,32 @@
+// Copyright Epic Games, Inc. All Rights Reserved.
+
+#pragma once
+
+#include "AndroidSingleInstanceServiceRuntimeSettings.generated.h"
+
+/**
+* Implements the settings for the AndroidSingleInstanceService plugin.
+*/
+
+
+UCLASS(Config = Engine, DefaultConfig)
+class UAndroidSingleInstanceServiceRuntimeSettings : public UObject
+{
+	GENERATED_UCLASS_BODY()
+
+	// Enable Android SingleInstanceService for packaged builds and quick launch
+	UPROPERTY(EditAnywhere, config, Category = Packaging, Meta = (DisplayName = "Use AndroidSingleInstanceService"))
+	bool bEnablePlugin;
+
+	// Compile standalone ASIS project
+	UPROPERTY(EditAnywhere, config, Category = Packaging, Meta=(EditCondition = "bEnablePlugin"))
+	bool bCompileASISProject;
+
+private:
+
+};
+
+#if UE_ENABLE_INCLUDE_ORDER_DEPRECATED_IN_5_2
+#include "CoreMinimal.h"
+#include "Engine/EngineTypes.h"
+#endif
diff --git a/Engine/Source/Programs/UnrealBuildTool/Platform/Android/UEBuildAndroid.cs b/Engine/Source/Programs/UnrealBuildTool/Platform/Android/UEBuildAndroid.cs
--- a/Engine/Source/Programs/UnrealBuildTool/Platform/Android/UEBuildAndroid.cs
+++ b/Engine/Source/Programs/UnrealBuildTool/Platform/Android/UEBuildAndroid.cs
@@ -418,6 +418,27 @@ namespace UnrealBuildTool
 		{
 		}
 
+		public static bool IsMakeAAREnabled(FileReference? ProjectFile, ILogger InLogger)
+		{
+			// look in ini settings for what platforms to compile for
+			DirectoryReference? ProjectDirectory = DirectoryReference.FromFile(ProjectFile);
+			ConfigHierarchy Ini = ConfigCache.ReadHierarchy(ConfigHierarchyType.Engine, DirectoryReference.FromFile(ProjectFile), UnrealTargetPlatform.Android);
+			bool bEnablePlugin = false;
+
+			if (Ini != null)
+			{
+				Ini.GetBool("/Script/AndroidSingleInstanceServiceEditor.AndroidSingleInstanceServiceRuntimeSettings", "bEnablePlugin", out bEnablePlugin);
+				// Only log if enabled to avoid adding extra logging to builds that don't use this feature.
+				if (bEnablePlugin)
+				{
+					InLogger.LogInformation("IsMakeAAREnabled {Ini} and {ProjectFile} , {ProjectDirectory}, {bEnablePlugin}", Ini, ProjectFile, ProjectDirectory, bEnablePlugin);
+				}
+			}
+
+			return bEnablePlugin;
+		}
+
+
 		public virtual void SetUpSpecificEnvironment(ReadOnlyTargetRules Target, CppCompileEnvironment CompileEnvironment, LinkEnvironment LinkEnvironment, ILogger Logger)
 		{
 			string NDKPath = Environment.GetEnvironmentVariable("NDKROOT")!;
@@ -448,6 +469,11 @@ namespace UnrealBuildTool
 
 			CompileEnvironment.Definitions.Add("WITH_EDITOR=0");
 			CompileEnvironment.Definitions.Add("USE_NULL_RHI=0");
+			if (IsMakeAAREnabled(Target.ProjectFile, Logger))
+			{
+				Logger.LogInformation("SetUpSpecificEnvironment is adding USE_ANDROID_STANDALONE because IsMakeAAREnabled({ProjectFile}) was true!", Target.ProjectFile);
+				CompileEnvironment.Definitions.Add("USE_ANDROID_STANDALONE=1");
+			}
 
 			DirectoryReference NdkDir = new DirectoryReference(NDKPath);
 			//CompileEnvironment.SystemIncludePaths.Add(DirectoryReference.Combine(NdkDir, "sources/cxx-stl/llvm-libc++/include"));
diff --git a/Engine/Source/Programs/UnrealBuildTool/Platform/Android/UEDeployAndroid.cs b/Engine/Source/Programs/UnrealBuildTool/Platform/Android/UEDeployAndroid.cs
--- a/Engine/Source/Programs/UnrealBuildTool/Platform/Android/UEDeployAndroid.cs
+++ b/Engine/Source/Programs/UnrealBuildTool/Platform/Android/UEDeployAndroid.cs
@@ -2665,7 +2665,18 @@ namespace UnrealBuildTool
 
 			StringBuilder Text = new StringBuilder();
 			Text.AppendLine(XML_HEADER);
-			Text.AppendLine("<manifest xmlns:android=\"http://schemas.android.com/apk/res/android\" xmlns:tools=\"http://schemas.android.com/tools\"");
+
+			bool bIsMakeAAREnabled = AndroidPlatform.IsMakeAAREnabled(ProjectFile, Logger);
+			if (bIsMakeAAREnabled)
+			{
+				// Note: dist was added for SingleInstanceService to allow instant module features, but need feedback if this will cause any issues since according to docs (https://developer.android.com/guide/playcore/feature-delivery/instant) An instant-enabled module cannot use background services. Additionally, such a module cannot send notifications when running in the background.
+				
+				Text.AppendLine("<manifest xmlns:android=\"http://schemas.android.com/apk/res/android\" xmlns:dist=\"http://schemas.android.com/apk/distribution\" xmlns:tools=\"http://schemas.android.com/tools\"");
+			}
+			else
+			{
+				Text.AppendLine("<manifest xmlns:android=\"http://schemas.android.com/apk/res/android\" xmlns:tools=\"http://schemas.android.com/tools\"");
+			}
 			//			Text.AppendLine(string.Format("          package=\"{0}\"", PackageName));
 			if (ExtraManifestNodeTags != null)
 			{
@@ -2696,6 +2707,14 @@ namespace UnrealBuildTool
 				Text.AppendLine("\t</queries>");
 			}
 
+			if (bIsMakeAAREnabled)
+			{
+				Text.AppendLine("\t<dist:module dist:instant=\"true\" />");
+				Text.AppendLine("\t<queries>");
+				Text.AppendLine(string.Format("\t\t<package android:name=\"{0}\" />", PackageName));
+				Text.AppendLine("\t</queries>");
+			}
+
 			Text.AppendLine("\t<!-- Application Definition -->");
 			Text.AppendLine("\t<application android:label=\"@string/app_name\"");
 			Text.AppendLine("\t             android:icon=\"@drawable/icon\"");
@@ -4382,6 +4401,7 @@ namespace UnrealBuildTool
 
 			// Initialize UPL contexts for each architecture enabled
 			UPL.Init(NDKArches, bForDistribution, EngineDirectory, IntermediateAndroidPath, ProjectDirectory, Configuration.ToString(), bSkipGradleBuild, bPerArchBuildDir: true, ArchRemapping: ArchRemapping);
+			UPL.SetGlobalContextVariable("AndroidPackageName", PackageName);
 
 			IEnumerable<Tuple<UnrealArch, string>>? BuildList = null;
 
diff --git a/Engine/Source/Runtime/ApplicationCore/Private/Android/AndroidPlatformFramePacer.cpp b/Engine/Source/Runtime/ApplicationCore/Private/Android/AndroidPlatformFramePacer.cpp
--- a/Engine/Source/Runtime/ApplicationCore/Private/Android/AndroidPlatformFramePacer.cpp
+++ b/Engine/Source/Runtime/ApplicationCore/Private/Android/AndroidPlatformFramePacer.cpp
@@ -13,7 +13,12 @@
 
 TAutoConsoleVariable<int32> FAndroidPlatformRHIFramePacer::CVarUseSwappyForFramePacing(
 	TEXT("a.UseSwappyForFramePacing"),
+#if USE_ANDROID_STANDALONE
+	// swappy crashes when using single instance and texture views.
+	0,
+#else
 	1,
+#endif
 	TEXT("True to use Swappy for frame pacing."));
 
 TAutoConsoleVariable<int32> FAndroidPlatformRHIFramePacer::CVarSupportNonVSyncMultipleFrameRates(
@@ -23,7 +28,12 @@ TAutoConsoleVariable<int32> FAndroidPlatformRHIFramePacer::CVarSupportNonVSyncMu
 
 TAutoConsoleVariable<int32> FAndroidPlatformRHIFramePacer::CVarAllowFrameTimestamps(
 	TEXT("a.AllowFrameTimestamps"),
+#if USE_ANDROID_STANDALONE
+	// frame time stamping when using texture views and single instance in android cannot be used as the frame numbering in android driver won't match up to updates.
+	0,
+#else
 	1,
+#endif
 	TEXT("True to allow the use use eglGetFrameTimestampsANDROID et al for frame pacing or spew."));
 
 TAutoConsoleVariable<int32> FAndroidPlatformRHIFramePacer::CVarTimeStampErrorRetryCount(
diff --git a/Engine/Source/Runtime/ApplicationCore/Private/Android/AndroidWindow.cpp b/Engine/Source/Runtime/ApplicationCore/Private/Android/AndroidWindow.cpp
--- a/Engine/Source/Runtime/ApplicationCore/Private/Android/AndroidWindow.cpp
+++ b/Engine/Source/Runtime/ApplicationCore/Private/Android/AndroidWindow.cpp
@@ -91,8 +91,13 @@ static void ClearCachedWindowRects()
 	CachedWindowRect_EventThread = FAndroidCachedWindowRectParams();
 }
 
-static int32 GSurfaceViewWidth = -1;
-static int32 GSurfaceViewHeight = -1;
+static int32 GSurfaceViewX = 0;
+static int32 GSurfaceViewY = 0;
+int32 GSurfaceViewWidth = -1;
+int32 GSurfaceViewHeight = -1;
+
+void* GAndroidWindowOverride = nullptr;
+static ANativeWindow* GAcquiredWindow = nullptr;
 
 void* FAndroidWindow::NativeWindow = NULL;
 
@@ -155,9 +160,15 @@ JNI_METHOD void Java_com_epicgames_unreal_GameActivity_nativeSetWindowInfo(JNIEn
 
 JNI_METHOD void Java_com_epicgames_unreal_GameActivity_nativeSetSurfaceViewInfo(JNIEnv* jenv, jobject thiz, jint width, jint height)
 {
-	GSurfaceViewWidth = width;
-	GSurfaceViewHeight = height;
-	UE_LOG(LogAndroid, Log, TEXT("nativeSetSurfaceViewInfo width=%d and height=%d"), GSurfaceViewWidth, GSurfaceViewHeight);
+	STANDALONE_DEBUG_LOG( TEXT("nativeSetSurfaceViewInfo prev width=%d and prev height=%d"), GSurfaceViewWidth, GSurfaceViewHeight);
+
+	if (GAndroidWindowOverride != nullptr && (width != GSurfaceViewWidth || height != GSurfaceViewHeight))	
+	{
+		GSurfaceViewWidth = width;
+		GSurfaceViewHeight = height;
+		FAppEventManager::GetInstance()->EnqueueAppEvent(APP_EVENT_STATE_WINDOW_RESIZED, FAppEventData((ANativeWindow*)GAndroidWindowOverride));
+		STANDALONE_DEBUG_LOG(TEXT("nativeSetSurfaceViewInfo width=%d and height=%d"), GSurfaceViewWidth, GSurfaceViewHeight);
+	}
 }
 
 JNI_METHOD void Java_com_epicgames_unreal_GameActivity_nativeSetSafezoneInfo(JNIEnv* jenv, jobject thiz, jboolean bIsPortrait, jfloat left, jfloat top, jfloat right, jfloat bottom)
@@ -181,6 +192,54 @@ JNI_METHOD void Java_com_epicgames_unreal_GameActivity_nativeSetSafezoneInfo(JNI
 #endif
 	UE_LOG(LogAndroid, Log, TEXT("nativeSetSafezoneInfo bIsPortrait=%d, left=%f, top=%f, right=%f, bottom=%f"), bIsPortrait ? 1 : 0, left, top, right, bottom);
 }
+
+#if USE_ANDROID_STANDALONE
+
+extern void GAndroidWindowLock_Lock(FString calledBy);
+extern void GAndroidWindowLock_Unlock(FString calledBy);
+
+JNI_METHOD void Java_com_epicgames_makeaar_GameActivityForMakeAAR_nativeSetSurfaceOverride(JNIEnv* jenv, jobject thiz, jobject surface, jint x, jint y)
+{
+	if (surface != 0)
+	{
+		GSurfaceViewX = x;
+		GSurfaceViewY = y;
+		void* prev = GAndroidWindowOverride;
+		GAndroidWindowOverride = (ANativeWindow*)ANativeWindow_fromSurface(jenv, surface);
+		STANDALONE_DEBUG_LOG(TEXT("nativeSetSurfaceOverride(makeaar) applied: prev to new %p -> %p, (%d, %d)"), prev, GAndroidWindowOverride, GSurfaceViewX, GSurfaceViewY);
+		if (prev != GAndroidWindowOverride)
+		{
+			//GAndroidWindowLock_Unlock("Java_com_epicgames_makeaar_GameActivityForMakeAAR_nativeSetSurfaceOverride");
+		}
+
+		//if (GAndroidWindowOverride != nullptr)
+		//{
+		//	STANDALONE_DEBUG_LOGf(LogAndroid, TEXT("nativeSetSurfaceOverride. FTaskGraphInterface::IsRunning() = %d, GAndroidWindowOverride=%x"), FTaskGraphInterface::IsRunning(), GAndroidWindowOverride);
+
+		//	if (FTaskGraphInterface::IsRunning())
+		//	{
+		//		FSimpleDelegateGraphTask::CreateAndDispatchWhenReady(
+		//			FSimpleDelegateGraphTask::FDelegate::CreateLambda([=]()
+		//				{
+		//					STANDALONE_DEBUG_LOGf(LogAndroid, TEXT("from nativeSetSurfaceOverride, trigger UnlockAndroidWindow. FTaskGraphInterface::IsRunning() = %d, GAndroidWindowOverride=%x"), FTaskGraphInterface::IsRunning(), GAndroidWindowOverride);
+
+		//					FPlatformMisc::UnlockAndroidWindow();
+		//				}), TStatId(), nullptr, ENamedThreads::GameThread);
+		//	}
+		//}
+
+	}
+	else
+	{
+		GAndroidWindowOverride = nullptr;
+		
+		STANDALONE_DEBUG_LOG(TEXT("nativeSetSurfaceOverride(makeaar) setting to null and lock window"));
+		//GAndroidWindowLock_Lock("Java_com_epicgames_makeaar_GameActivityForMakeAAR_nativeSetSurfaceOverride");
+	}
+}
+
+#endif // USE_ANDROID_STANDALONE
+
 #endif
 
 bool FAndroidWindow::bAreCachedNativeDimensionsValid = false;
@@ -210,13 +269,62 @@ void FAndroidWindow::InvalidateCachedScreenRect()
 void FAndroidWindow::AcquireWindowRef(ANativeWindow* InWindow)
 {
 #if USE_ANDROID_JNI
+	STANDALONE_DEBUG_LOG(TEXT("AcquireWindowRef USE_ANDROID_JNI is enabled: InWindow=%p, GAcquiredWindow=%p, GAndroidWindowOverride=%p"), InWindow, GAcquiredWindow, GAndroidWindowOverride);
+
+	if (InWindow == nullptr)
+	{
+		UE_LOG(LogAndroid, Log, TEXT("FAndroidWindow::AcquireWindowRef skipped because InWindow is null."));
+		return;
+	}
+
+	if (GAcquiredWindow == InWindow)
+	{
+		STANDALONE_DEBUG_LOG(TEXT("AcquireWindowRef USE_ANDROID_JNI is enabled: %p and GAcquiredWindow == InWindow"), InWindow);
+		return;
+	}
+
+	if (GAcquiredWindow != nullptr)
+	{
+		STANDALONE_DEBUG_LOG(TEXT("AcquireWindowRef USE_ANDROID_JNI is enabled: %p and GAcquiredWindow != nullptr"), InWindow);
+
+		ReleaseWindowRef(GAcquiredWindow);
+	}
+
+	check(GAcquiredWindow == NULL);
+	STANDALONE_DEBUG_LOG(TEXT("AcquireWindowRef USE_ANDROID_JNI is enabled: %p and ANativeWindow_acquire"), InWindow);
+
 	ANativeWindow_acquire(InWindow);
+
+	// Added logic to store the Acquired window and when calling ReleaseWindowRef, check if the window is the same and only release if it matches
+	// This logic is to deal with the fact Android lifecycles for activities can overlap. ideally we would create a context based container to manage this
+	// but for now this is a useful protection.
+	GAcquiredWindow = InWindow;
+	STANDALONE_DEBUG_LOG(TEXT("FAndroidWindow::AcquireWindowRef overrode window: %p"), GAndroidWindowOverride);
+
+#else
+	STANDALONE_DEBUG_LOG(TEXT("AcquireWindowRef USE_ANDROID_JNI is NOT enabled: %p"), InWindow);
+
 #endif
 }
 
 void FAndroidWindow::ReleaseWindowRef(ANativeWindow* InWindow)
 {
 #if USE_ANDROID_JNI
+	if (GAcquiredWindow == nullptr && InWindow == nullptr)
+	{
+		STANDALONE_DEBUG_LOG(TEXT("ReleaseWindowRef skipped because GAcquiredWindow is null.  Window %p reference will not be released."), InWindow);
+		return;
+	}
+
+	ANativeWindow* ReleaseWindow = GAcquiredWindow;
+	if (InWindow == nullptr || GAcquiredWindow == InWindow)
+	{
+		InWindow = GAcquiredWindow;
+		GAcquiredWindow = nullptr;
+	}
+
+	STANDALONE_DEBUG_LOG(TEXT("ReleaseWindowRef using window: %p"), InWindow);
+
 	ANativeWindow_release(InWindow);
 #endif
 }
@@ -226,11 +334,22 @@ void FAndroidWindow::ReleaseWindowRef(ANativeWindow* InWindow)
 #if USE_ANDROID_EVENTS
 	check(IsInAndroidEventThread());
 #endif
-	NativeWindow = InWindow; //using raw native window handle for now. Could be changed to use AndroidWindow later if needed
+	//using raw native window handle for now. Could be changed to use AndroidWindow later if needed
+	NativeWindow = InWindow; 
 }
 
 void* FAndroidWindow::GetHardwareWindow_EventThread()
 {
+#if USE_ANDROID_STANDALONE
+	if (NativeWindow && GAndroidWindowOverride)
+	{
+		STANDALONE_DEBUG_LOG(TEXT("GetHardwareWindow_EventThread overrode window: %p"), GAndroidWindowOverride);
+		return GAndroidWindowOverride;
+	}
+#endif
+
+	STANDALONE_DEBUG_LOG(TEXT("GetHardwareWindow_EventThread NativeWindow: %p"), NativeWindow);
+
 	return NativeWindow;
 }
 
@@ -280,6 +399,14 @@ void FAndroidWindow::EventManagerUpdateWindowDimensions(int32 Width, int32 Heigh
 	check(bAreCachedNativeDimensionsValid);
 	check(Width >= 0 && Height >= 0);
 
+#if USE_ANDROID_STANDALONE
+	if (GAndroidWindowOverride && GSurfaceViewWidth > 0)
+	{
+		Width = GSurfaceViewWidth;
+		Height = GSurfaceViewHeight;
+	}
+#endif
+
 	bool bChanged = CachedNativeWindowWidth != Width || CachedNativeWindowHeight != Height;
 
 	CachedNativeWindowWidth = Width;
diff --git a/Engine/Source/Runtime/Core/Private/Android/AndroidJavaEnv.cpp b/Engine/Source/Runtime/Core/Private/Android/AndroidJavaEnv.cpp
--- a/Engine/Source/Runtime/Core/Private/Android/AndroidJavaEnv.cpp
+++ b/Engine/Source/Runtime/Core/Private/Android/AndroidJavaEnv.cpp
@@ -59,7 +59,7 @@ private:
 	}
 };
 
-void AndroidJavaEnv::InitializeJavaEnv( JavaVM* VM, jint Version, jobject GlobalThis )
+void AndroidJavaEnv::InitializeJavaEnv( JavaVM* VM, jint Version, jobject GlobalThis)
 {
 	if (CurrentJavaVM == nullptr)
 	{
@@ -67,7 +67,7 @@ void AndroidJavaEnv::InitializeJavaEnv( JavaVM* VM, jint Version, jobject Global
 		CurrentJavaVersion = Version;
 
 		JNIEnv* Env = GetJavaEnv(false);
-		jclass MainClass = Env->FindClass("com/epicgames/unreal/GameActivity");
+		jclass MainClass = Env->FindClass(ANDROID_GAMEACTIVITY_CLASSPATH);
 		jclass classClass = Env->FindClass("java/lang/Class");
 		jclass classLoaderClass = Env->FindClass("java/lang/ClassLoader");
 		jmethodID getClassLoaderMethod = Env->GetMethodID(classClass, "getClassLoader", "()Ljava/lang/ClassLoader;");
diff --git a/Engine/Source/Runtime/Core/Private/Android/AndroidPlatformFile.cpp b/Engine/Source/Runtime/Core/Private/Android/AndroidPlatformFile.cpp
--- a/Engine/Source/Runtime/Core/Private/Android/AndroidPlatformFile.cpp
+++ b/Engine/Source/Runtime/Core/Private/Android/AndroidPlatformFile.cpp
@@ -107,6 +107,9 @@ FString GExternalFilePath;
 // External font path base - setup during load
 FString GFontPathBase;
 
+// Last opened OBB comment (set during mounting of OBB)
+FString GLastOBBComment;
+
 // Is the OBB in an APK file or not
 bool GOBBinAPK;
 FString GAPKFilename;
@@ -126,11 +129,17 @@ JNI_METHOD void Java_com_epicgames_unreal_GameActivity_nativeSetObbInfo(JNIEnv*
 	GAndroidProjectName = FJavaHelper::FStringFromParam(jenv, ProjectName);
 	GPackageName = FJavaHelper::FStringFromParam(jenv, PackageName);
 	GAndroidAppType = FJavaHelper::FStringFromParam(jenv, AppType);
-	
+
 	GAndroidPackageVersion = Version;
 	GAndroidPackagePatchVersion = PatchVersion;
 }
 
+//This function is declared in the Java-defined class, GameActivity.java: "public native String nativeGetObbComment();"
+JNI_METHOD jstring Java_com_epicgames_makeaar_GameActivityForMakeAAR_nativeGetObbComment(JNIEnv* jenv, jobject thiz)
+{
+	return jenv->NewStringUTF(TCHAR_TO_UTF8(*GLastOBBComment));
+}
+
 // Constructs the base path for any files which are not in OBB/pak data
 const FString &GetFileBasePath()
 {
@@ -761,6 +770,17 @@ public:
 		check( DirOffset + DirSize <= FileLength );
 		check( NumEntries > 0 );
 
+		uint16 CommentLength = (Buffer.GetValue<uint16>(EOCDIndex + kEOCDCommentLen));
+		if (CommentLength > 0)
+		{
+			GLastOBBComment = FString(CommentLength, reinterpret_cast<const ANSICHAR*>(Buffer.Data + EOCDIndex + kEOCDCommentStart));
+		}
+		else
+		{
+			GLastOBBComment = FString("");
+		}
+
+
 		/*
 		* Walk through the central directory, adding entries to the hash table.
 		*/
@@ -930,6 +950,8 @@ private:
 	const uint32 kEOCDNumEntries = 8; // offset to #of entries in file
 	const uint32 kEOCDSize = 12; // size of the central directory
 	const uint32 kEOCDFileOffset = 16; // offset to central directory
+	const uint32 kEOCDCommentLen = 20; // offset to comment length (ushort)
+	const uint32 kEOCDCommentStart = 22; // offset to start of optional comment
 
 	const uint32 kMaxCommentLen = 65535; // longest possible in ushort
 	const uint32 kMaxEOCDSearch = (kMaxCommentLen + kEOCDLen);
