fix missing include.
#android

diff --git a/Engine/Plugins/Runtime/AndroidSingleInstanceService/Source/AndroidSingleInstanceService/Private/AndroidSingleInstanceServiceBPLibrary.cpp b/Engine/Plugins/Runtime/AndroidSingleInstanceService/Source/AndroidSingleInstanceService/Private/AndroidSingleInstanceServiceBPLibrary.cpp
--- a/Engine/Plugins/Runtime/AndroidSingleInstanceService/Source/AndroidSingleInstanceService/Private/AndroidSingleInstanceServiceBPLibrary.cpp
+++ b/Engine/Plugins/Runtime/AndroidSingleInstanceService/Source/AndroidSingleInstanceService/Private/AndroidSingleInstanceServiceBPLibrary.cpp
@@ -47,6 +47,7 @@ REGISTER_DATA_RECEIVED_DELEGATE(DoubleArray);
 #include "Android/AndroidJNI.h"
 #include "Android/AndroidApplication.h"
 #include "Android/AndroidWindow.h"
+#include "Async/TaskGraphInterfaces.h"
 
 #if USE_ANDROID_JNI
 #include <jni.h>
