Android SingleInstance Services fixes to address crash and activity/surface binding issues.
known issues remaining: have to pause before switching surface binding (should be handled internally eventually)
need to finish up the changes for compiling reference ASIS android studio project samlple. This will be in another submit.
#android

diff --git a/Engine/Build/Android/Java/src/com/epicgames/unreal/GameActivity.java.template b/Engine/Build/Android/Java/src/com/epicgames/unreal/GameActivity.java.template
--- a/Engine/Build/Android/Java/src/com/epicgames/unreal/GameActivity.java.template
+++ b/Engine/Build/Android/Java/src/com/epicgames/unreal/GameActivity.java.template
@@ -426,7 +426,7 @@ public class GameActivity extends $${gameActivitySuperClass}$$ implements Surfac
 	int DepthBufferPreference = 0;
 	
 	// PropagatetAlpha from the manifest
-	int PropagateAlpha = 0;
+	protected int PropagateAlpha = 0;
 
 	/** AssetManger reference - populated on start up and used when the OBB is packed into the APK */
 	protected AssetManager			AssetManagerReference;
diff --git a/Engine/Plugins/Runtime/AndroidSingleInstanceService/Source/AndroidSingleInstanceService/Android/src/main/java/com/epicgames/makeaar/Engine.java b/Engine/Plugins/Runtime/AndroidSingleInstanceService/Source/AndroidSingleInstanceService/Android/src/main/java/com/epicgames/makeaar/Engine.java
--- a/Engine/Plugins/Runtime/AndroidSingleInstanceService/Source/AndroidSingleInstanceService/Android/src/main/java/com/epicgames/makeaar/Engine.java
+++ b/Engine/Plugins/Runtime/AndroidSingleInstanceService/Source/AndroidSingleInstanceService/Android/src/main/java/com/epicgames/makeaar/Engine.java
@@ -3,8 +3,6 @@ package com.epicgames.makeaar;
 import static android.content.Context.CONTEXT_IGNORE_SECURITY;
 import static android.content.Context.CONTEXT_INCLUDE_CODE;
 
-import androidx.annotation.NonNull;
-import androidx.annotation.UiThread;
 import android.annotation.SuppressLint;
 import android.app.Activity;
 import android.app.ActivityManager;
@@ -214,8 +212,7 @@ public class Engine
 
 			Engine engine = Get();
 
-			//synchronized(syncLock)
-			{
+			synchronized (syncLock) {
 				if (holder.getSurface() != engine.renderSurface && engine.renderSurface != null)
 				{
 					Log.w(TAG, "surfaceChanged and releaseVolatileResources cause holder surface is different than prev");
@@ -254,8 +251,7 @@ public class Engine
 				+ ", holder.surface=" + holder.getSurface()
 			);
 
-			//synchronized(syncLock)
-			{
+			synchronized (syncLock) {
 				if (engine.renderSurfaceValid && engine.renderSurface == holder.getSurface()) {
 					Log.i(TAG, "surfaceDestroyed, doing cleanup and pausing, holder = " + holder
 							+ ", lifecycleContextID=" + lifecycleContextID
@@ -310,8 +306,7 @@ public class Engine
 		@Override
 		public void onSurfaceTextureAvailable(SurfaceTexture surfaceTexture, int width, int height) {
 
-			//synchronized(syncLock)
-			{
+			synchronized (syncLock) {
 				TextureView textureView = ((TextureView) view);
 
 				int[] surfacePosition = new int[]{0, 0};
@@ -336,8 +331,7 @@ public class Engine
 
 		@Override
 		public void onSurfaceTextureSizeChanged(SurfaceTexture surfaceTexture, int width, int height) {
-			//synchronized(syncLock)
-			{	
+			synchronized (syncLock) {
 				TextureView textureView = ((TextureView) view);
 
 				int[] surfacePosition = new int[]{0, 0};
@@ -357,8 +351,7 @@ public class Engine
 		@Override
 		public boolean onSurfaceTextureDestroyed(SurfaceTexture surfaceTexture) {
 
-			//synchronized(syncLock)
-			{
+			synchronized (syncLock) {
 
 				Log.i(TAG, "onSurfaceTextureDestroyed, surfaceForTextureView = " + surfaceForTextureView + ", renderSurface = " + Engine.Get().renderSurface);
 
@@ -558,8 +551,7 @@ public class Engine
 
 	public boolean ApplyPendingSurfaceChanges()
 	{
-		//synchronized(syncLock)
-		{
+		synchronized (syncLock) {
 			if (pendingChangedSurface == null || pendingChangedSurface == renderSurface) {
 				Log.w(TAG, "logOrder:" + logOrder + " *** Engine::ApplyPendingSurfaceChanges(DO NOTHING pending is NULL, or same as current) - renderSurfaceValid=" + renderSurfaceValid + ", renderSurface=" + renderSurface);
 				return false;
@@ -657,7 +649,7 @@ public class Engine
 
 		if (bAcitivityNeedsSetup)
 		{
-			//synchronized(syncLock)
+			synchronized(syncLock)
 			{
 				if (GameActivity.Get() == null)
 				{
@@ -692,7 +684,7 @@ public class Engine
 			Log.e(TAG, "logOrder:" + logOrder + " ** Engine::onSurfaceTextureAvailable(ERROR SURFACE POSSIBLY ALREADY EXISTS) - surfaceChanged size=" + Arrays.toString(gameActivitySetupInfo.ViewSize) + ", bResuming=" + bResuming + ", renderSurfaceValid=" + renderSurfaceValid + "renderSurface" + renderSurface + ", surfaceCreatedCounter=" + surfaceCreatedCounter );
 		}
 
-		//synchronized(syncLock)
+		synchronized(syncLock)
 		{
 			pendingChangedSurface = gameActivitySetupInfo.renderSurface;
 		}
@@ -796,7 +788,7 @@ public class Engine
 
 		if (bAcitivityNeedsSetup)
 		{
-			//synchronized(syncLock)
+			synchronized(syncLock)
 			{
 				if (GameActivity.Get() == null)
 				{
@@ -871,7 +863,10 @@ public class Engine
 
 	public boolean AttachExternalRenderSurface(int contextID, Surface externalSurface, int viewPos[], int viewSize[])
 	{
-		++logOrder;
+		if (!bHasBeenInit || GameActivity.Get() == null || externalSurface == null)
+		{
+			return false;
+		}
 
 
 		Log.i(TAG, "** Engine::AttachExternalRenderSurface() - bFirstInit"
@@ -892,79 +887,25 @@ public class Engine
 		);
 
 
-		if (!bHasBeenInit || GameActivity.Get() == null || externalSurface == null) {
-			Log.e(TAG, "** Engine::AttachExternalRenderSurface() - ERROR, RETURN NO CHANGE");
-			return false;
+		if (false == VerifyContextID(contextID, "Engine::AttachExternalRenderSurface()")) {
+			onStop(lifecycleContextID, "Engine::AttachExternalRenderSurface");
+			lifecycleContextID = contextID;
 		}
 
-		if (onPauseCounter > 0 || startCounter != 0 || renderSurface != null || pendingChangedSurface != null)
+		else if (renderSurface != null || surfaceCreatedCounter>0)
 		{
-			Log.w(TAG, "** Engine::AttachExternalRenderSurface() - CLEANUP previous context"
-				+ ", bHasBeenInit:" + bHasBeenInit
-				+ ", startCounter:" + startCounter
-				+ ", bResuming:" + bResuming
-				+ ", onPauseCounter:" + onPauseCounter
-				+ ", renderSurfaceValid:" + renderSurfaceValid
-				+ ", renderSurface:" + renderSurface
-				+ ", surfaceCreatedCounter:" + surfaceCreatedCounter
-				+ ", pendingChangedSurface:" + pendingChangedSurface
-			);
+			releaseVolatileResources(contextID, "AttachExternalRenderSurface renderSurface != null || surfaceCreatedCounter>0");
 
-			if (bHasBeenInit) {
-				if (onPauseCounter<=0) {
-					QueuePause("from AttachExternalRenderSurface onPauseCounter=" + onPauseCounter);
-				}
-				if (startCounter>0)
-				{
-					QueueStop( "from AttachExternalRenderSurface startCounter=" + startCounter);
-				}
-
-				//GetOverrideGameActivity().nativeSetSurfaceViewInfo(1, 1);
-				//GetOverrideGameActivity().nativeSetSurfaceOverride(null, renderSurfaceViewPosition[0], renderSurfaceViewPosition[1]);
-				//GetOverrideGameActivity().nativeAppCommand(APP_CMD_TERM_WINDOW);
-				//GetOverrideGameActivity().nativeAppCommand(APP_CMD_PAUSE);
-				//GetOverrideGameActivity().nativeAppCommand(APP_CMD_STOP);
-				
-//
-//				bCalledStart = false;
-//				startCounter = 0;
-//				bResuming = false;
-//				onPauseCounter = 0;
-//				renderSurfaceValid = false;
-//				renderSurface = null;
-//				surfaceCreatedCounter = 0;
-//				pendingChangedSurface = null;
-//				lifecycleContextID = contextID;
-			}
 		}
 
-
-
-
-		//if (false == VerifyContextID(contextID, "Engine::AttachExternalRenderSurface, contextID=" + contextID)) {
-		//		boolean shouldResume = renderSurface != null || pendingChangedSurface != null || bResuming;
-		//	//onStop(lifecycleContextID, "Engine::AttachExternalRenderSurface");
-		//	releaseVolatileResources(lifecycleContextID, "AttachExternalRenderSurface lifecycleContextID != contextID");
-		//	lifecycleContextID = contextID;
-		//	bResuming = shouldResume;
-		//}
-
-		//else if (renderSurface != null || surfaceCreatedCounter > 0)
-		{
-			//releaseVolatileResources(contextID, "AttachExternalRenderSurface renderSurface != null || surfaceCreatedCounter>0");
-		}
-
-		
-
 		Log.i(TAG, "** Engine::AttachExternalRenderSurface() - CHECK 1A");
+		++logOrder;
 
 
-
-		synchronized(syncLock)
+		//synchronized(syncLock)
 		{
-			Log.i(TAG, "** Engine::AttachExternalRenderSurface() - CHECK 1B");
+		Log.i(TAG, "** Engine::AttachExternalRenderSurface() - CHECK 1B");
 			pendingChangedSurface = externalSurface;
-			bResuming = true;
 			if (viewSize.length==2)
 			{
 				renderSurfaceViewSize[0] = viewSize[0];
@@ -986,7 +927,6 @@ public class Engine
 			}
 		}
 
-		QueueStart("AttachExternalRenderSurface");
 
 		Log.i(TAG, "** Engine::AttachExternalRenderSurface() - about to ApplyPendingSurfaceChanges...");
 
@@ -998,7 +938,7 @@ public class Engine
 		{
 			Log.i(TAG, "** Engine::AttachExternalRenderSurface() - about to QueueResume ...");
 
-			QueueResume( "AttachExternalRenderSurface");
+			QueueResume("AttachExternalRenderSurface");
 		}
 
 		Log.i(TAG, "Engine::AttachExternalRenderSurface(DONE)"
@@ -1014,8 +954,7 @@ public class Engine
 		return bResuming;
 	}
 
-	@UiThread
-public static synchronized void QueuePause(String reasonString)
+	public static void QueuePause(String reasonString)
 	{
 		if (bResuming && onPauseCounter <= 0) {
 			Log.e(TAG, "QueuePause is CALLED with lifecycleContextID=" + lifecycleContextID + ", bResuming=" + bResuming);
@@ -1025,12 +964,13 @@ public static synchronized void QueuePause(String reasonString)
 
 		if (GameActivity.Get() == null) {
 			Log.e(TAG, "QueuePause is REJECTED cause GameActivity is not valid with lifecycleContextID=" + lifecycleContextID + ", bResuming=" + bResuming + ", reasonString=" + reasonString);
+
 			return;
 		}
 
-		GetOverrideGameActivity().GetCurrentActivityContext().runOnUiThread(() ->
+		GetOverrideGameActivity().GetCurrentActivityContext().runOnMainThread(() ->
 				{
-					//synchronized(syncLock)
+					synchronized(syncLock)
 					{
 						Log.i(TAG, "QueuePause is calling engine.onPause with lifecycleContextID=" + lifecycleContextID + ", bResuming=" + bResuming);
 						Engine.Get().onPause(lifecycleContextID, reasonString);
@@ -1039,8 +979,7 @@ public static synchronized void QueuePause(String reasonString)
 		);
 	}
 
-	@UiThread
-	public static synchronized void QueueResume(String reasonString)
+	public static void QueueResume(String reasonString)
 	{
 		if (!GameActivityForMakeAAR.isValidGameActivity())
 		{
@@ -1065,8 +1004,7 @@ public static synchronized void QueuePause(String reasonString)
 		);
 	}
 
-	@UiThread
-	public static synchronized void QueueStart(String reasonString)
+	public static void QueueStart(String reasonString)
 	{
 		if (!GameActivityForMakeAAR.isValidGameActivity())
 		{
@@ -1074,20 +1012,23 @@ public static synchronized void QueuePause(String reasonString)
 			return;
 		}
 
-		GetOverrideGameActivity().GetCurrentActivityContext().runOnUiThread(()->
+		GetOverrideGameActivity().GetCurrentActivityContext().runOnMainThread(new Runnable()
+		{
+			@Override
+			public void run()
 			{
-				//synchronized (syncLock)
+				synchronized (syncLock)
 				{
 					Log.i(TAG, "QueueStart is calling engine.onStart() with lifecycleContextID=" + lifecycleContextID + ", reasonString=" + reasonString);
 	
 					Engine.Get().onStart(lifecycleContextID, "QueueStart with reasonString=" + reasonString);
 				}
 			}
-		);
+		});
 	}
 
-	@UiThread
-	public static synchronized void QueueStop(String reasonString)
+
+	public static void QueueStop(String reasonString)
 	{
 		if (false == GameActivityForMakeAAR.isValidGameActivity())
 		{
@@ -1095,16 +1036,19 @@ public static synchronized void QueuePause(String reasonString)
 			return;
 		}
 
-		GetOverrideGameActivity().GetCurrentActivityContext().runOnUiThread(()->
+		GetOverrideGameActivity().GetCurrentActivityContext().runOnMainThread(new Runnable()
+		{
+			@Override
+			public void run()
 			{
-				//synchronized(syncLock)
+				synchronized(syncLock)
 				{
 					Log.i(TAG, "QueueStop is calling engine.onStop with lifecycleContextID=" + lifecycleContextID + ", reasonString=" + reasonString);
-					Engine.Get().onStop(lifecycleContextID, "QueueStop with reasonString=" + reasonString);
+			Engine.Get().onStop(lifecycleContextID, "QueueStop with reasonString=" + reasonString);
 				}
 	
 			}
-		);
+		});
 	}
 
 	boolean bHasBeenInit = false;
@@ -1359,7 +1303,7 @@ public static synchronized void QueuePause(String reasonString)
 		lifecycleContextID = (lifecycleContextID%100) + (++lifecycleActivityContextCounter) * 100;
 	}
 
-	private void onPause(int contextID, String reasonString)
+	public void onPause(int contextID, String reasonString)
 	{
 		if (contextID != 0 && lifecycleContextID != 0 && contextID != lifecycleContextID) {
 			Log.e(TAG, "[lifecycleContextID=" + lifecycleContextID + "] contextID:" + contextID + " ABNORMAL Engine::onPause(contextID) - onPause called. with contextID change = " + contextID + ", reasonString=" + reasonString);
@@ -1389,10 +1333,9 @@ public static synchronized void QueuePause(String reasonString)
 			--surfaceCreatedCounter;
 		}
 
-		if (renderSurface != null || pendingChangedSurface != null)
-		{
+		if (renderSurface != null || pendingChangedSurface != null) {
 			GetOverrideGameActivity().nativeSetSurfaceOverride(null, renderSurfaceViewPosition[0], renderSurfaceViewPosition[1]);
-			//pendingChangedSurface = pendingChangedSurface != null ? pendingChangedSurface : renderSurface;
+			pendingChangedSurface = pendingChangedSurface != null ? pendingChangedSurface : renderSurface;
 			renderSurface = null;
 			renderSurfaceValid = false;
 			bResuming |= pendingChangedSurface != null;
@@ -1401,13 +1344,12 @@ public static synchronized void QueuePause(String reasonString)
 
 		GetOverrideGameActivity().nativeAppCommand(APP_CMD_LOST_FOCUS);
 		GetOverrideGameActivity().nativeAppCommand(APP_CMD_PAUSE);
-		//if (pendingChangedSurface != null)
-		//	GetOverrideGameActivity().nativeAppCommand(APP_CMD_TERM_WINDOW);
+		//GetOverrideGameActivity().nativeAppCommand(APP_CMD_TERM_WINDOW);
 
 		//releaseVolatileResources("onPause called with reason=" + reasonString);
 	}
 
-	private void onResume(int contextID, String reasonString)
+	public void onResume(int contextID, String reasonString)
 	{
 		if (contextID != 0 && lifecycleContextID != 0 && contextID != lifecycleContextID) {
 			Log.e(TAG, "[lifecycleContextID=" + lifecycleContextID + "] contextID:" + contextID + " ABNORMAL Engine::onResume - onResume called. with contextID change = " + contextID + ", reasonString=" + reasonString);
@@ -1490,7 +1432,7 @@ public static synchronized void QueuePause(String reasonString)
 		}
 	}
 
-	private void onStart(int contextID, String reasonString)
+	public void onStart(int contextID, String reasonString)
 	{
 		if (contextID != 0 && lifecycleContextID != 0 && contextID != lifecycleContextID) {
 			Log.e(TAG, "[lifecycleContextID=" + lifecycleContextID + "] contextID:" + contextID + " ABNORMAL Engine::onStart(contextID) - onStart called. with contextID change = " + contextID + ", reasonString=" + reasonString);
@@ -1517,7 +1459,7 @@ public static synchronized void QueuePause(String reasonString)
 		}
 	}
 
-	private void onStop(int contextID, String reasonString)
+	public void onStop(int contextID, String reasonString)
 	{
 		if (contextID != 0 && lifecycleContextID != 0 && contextID != lifecycleContextID) {
 			Log.e(TAG, "[lifecycleContextID=" + lifecycleContextID + "] contextID:" + contextID + " ABNORMAL Engine::onStop(contextID) - onStop called. with contextID change = " + contextID + ", reasonString=" + reasonString);
@@ -1541,8 +1483,7 @@ public static synchronized void QueuePause(String reasonString)
 		}
 
 		//releaseVolatileResources("onStop() called so release all resources");
-		//synchronized(syncLock)
-		{
+		synchronized (syncLock) {
 			pendingChangedSurface = null;
 		}
 		
@@ -1561,16 +1502,7 @@ public static synchronized void QueuePause(String reasonString)
 			lifecycleContextID = 0;
 		}
 
-		Log.i(TAG, "[lifecycleContextID=" + lifecycleContextID + "] logOrder:" + logOrder + " * Engine::onStop - onStop COMPLETED. bCalledStart was = " + bCalledStart
-			+ ", reasonString=" + reasonString
-			+ ", bHasBeenInit = " + bHasBeenInit
-			+ ", prev startCounter = " + startCounter
-			+ ", surfaceCreatedCounter = " + surfaceCreatedCounter
-			+ ", bResuming was = " + bResuming
-			+ ", onPauseCounter=" + onPauseCounter
-			+ ", renderSurfaceValid: " + renderSurfaceValid
-			+ ", renderSurface = " + renderSurface
-		);
+
 	}
 
 	private void ReleaseRenderSurfaces(int contextID, String reasonString)
@@ -1663,8 +1595,7 @@ public static synchronized void QueuePause(String reasonString)
 
 		}
 
-		//if (renderSurface != null)
-		{
+		if (renderSurface != null) {
 			GetOverrideGameActivity().nativeAppCommand(APP_CMD_TERM_WINDOW);
 		}
 
diff --git a/Engine/Plugins/Runtime/AndroidSingleInstanceService/Source/AndroidSingleInstanceService/Android/src/main/java/com/epicgames/makeaar/GameActivityForMakeAAR.java b/Engine/Plugins/Runtime/AndroidSingleInstanceService/Source/AndroidSingleInstanceService/Android/src/main/java/com/epicgames/makeaar/GameActivityForMakeAAR.java
--- a/Engine/Plugins/Runtime/AndroidSingleInstanceService/Source/AndroidSingleInstanceService/Android/src/main/java/com/epicgames/makeaar/GameActivityForMakeAAR.java
+++ b/Engine/Plugins/Runtime/AndroidSingleInstanceService/Source/AndroidSingleInstanceService/Android/src/main/java/com/epicgames/makeaar/GameActivityForMakeAAR.java
@@ -39,7 +39,6 @@ public class GameActivityForMakeAAR extends com.epicgames.unreal.GameActivity
 	
 
 	String _OBBFilename = null;
-	boolean _enablePropagateAlpha = false;
 
 	public static com.epicgames.unreal.Logger Log = new Logger("UE_ASIS", "MakeAARBaseForGameActivity");
 	
@@ -63,9 +62,9 @@ public class GameActivityForMakeAAR extends com.epicgames.unreal.GameActivity
 		else
 		{
 			GameActivityForMakeAAR newInstance = new GameActivityForMakeAAR();
+			
 			Log.debug("GameActivityForMakeAAR::Instance() constructor called with prev _gameActivityInstance=" + _gameActivityInstance + ", newInstance.activityContext=" + newInstance.activityContext + ", newInstance=" + newInstance + ", newInstance.isStandalone=" + newInstance.isStandalone);
 			_gameActivityInstance = newInstance;
-			
 		}
 
 		assert( _gameActivityInstance instanceof GameActivityForMakeAAR);
@@ -95,6 +94,8 @@ public class GameActivityForMakeAAR extends com.epicgames.unreal.GameActivity
 	}
 
 	@Override
+
+
 	protected void RestartApplication(String RestartExtra)
 	{
 		super.RestartApplication(RestartExtra);
@@ -234,7 +235,6 @@ public class GameActivityForMakeAAR extends com.epicgames.unreal.GameActivity
 		nativeSetAndroidStartupState(bDebuggerAttached);
 
 		_OBBFilename = OBBFilename;
-		_enablePropagateAlpha = enablePropagateAlpha;
 
 		GameActivity.setOBBInAPK(OBBFilename.equals(""));
 		Log.debug("APK path: " + activityContext.getPackageResourcePath());
@@ -401,9 +401,8 @@ public class GameActivityForMakeAAR extends com.epicgames.unreal.GameActivity
 		nativeSetAndroidStartupState(bDebuggerAttached);
 
 		_OBBFilename = gameActivitySetupInfo._OBBFilename;
-		int PropagateAlpha = gameActivitySetupInfo.PropagateAlpha;
 
-		_enablePropagateAlpha = gameActivitySetupInfo.PropagateAlpha == 0 ? false : true;
+		PropagateAlpha = gameActivitySetupInfo.PropagateAlpha;
 
 		GameActivity.setOBBInAPK(_OBBFilename.equals(""));
 
@@ -547,7 +546,7 @@ public class GameActivityForMakeAAR extends com.epicgames.unreal.GameActivity
 
 		//Activity _activity = null;
 		_OBBFilename = gameActivitySetupInfo._OBBFilename;
-		_enablePropagateAlpha = gameActivitySetupInfo.PropagateAlpha == 1;
+		PropagateAlpha = gameActivitySetupInfo.PropagateAlpha;
 		_bundle = gameActivitySetupInfo._bundle;
 		_extrasBundle = gameActivitySetupInfo._extrasBundle;
 		PackageDataInsideApkValue = gameActivitySetupInfo.packageDataInsideApkValue ? 1 : 0;
diff --git a/Engine/Plugins/Runtime/AndroidSingleInstanceService/Source/AndroidSingleInstanceService/Android/src/main/java/com/epicgames/makeaar/GameActivitySetupInfo.java b/Engine/Plugins/Runtime/AndroidSingleInstanceService/Source/AndroidSingleInstanceService/Android/src/main/java/com/epicgames/makeaar/GameActivitySetupInfo.java
--- a/Engine/Plugins/Runtime/AndroidSingleInstanceService/Source/AndroidSingleInstanceService/Android/src/main/java/com/epicgames/makeaar/GameActivitySetupInfo.java
+++ b/Engine/Plugins/Runtime/AndroidSingleInstanceService/Source/AndroidSingleInstanceService/Android/src/main/java/com/epicgames/makeaar/GameActivitySetupInfo.java
@@ -13,13 +13,12 @@ import android.view.Surface;
 import android.content.pm.PackageManager;
 import android.os.Parcel;
 import android.os.Parcelable;
-import android.util.Log;
 
 import java.io.File;
 import java.util.Map;
 import java.util.HashMap;
 
-//import com.epicgames.unreal.Logger;
+import com.epicgames.unreal.Logger;
 
 import static android.content.Context.CLIPBOARD_SERVICE;
 import static android.content.Context.WINDOW_SERVICE;
@@ -30,8 +29,7 @@ import com.epicgames.unreal.SimpleContextWrapper;
 
 public class GameActivitySetupInfo implements Parcelable {
 
-	//public static Logger Log = new Logger("UES", "GameActivitySetupInfo");
-	final static String TAG = "UE_ASIS_GameActivitySetupInfo";
+	public static Logger Log = new Logger("UES", "GameActivitySetupInfo");
 
 	public static final int MSG_ATTACH_EXTERNAL_SURFACE = 2;
 		
@@ -97,7 +95,7 @@ public class GameActivitySetupInfo implements Parcelable {
 							if (index >= 0)
 							{
 								ref.VulkanLevel = Integer.parseInt(dump.substring(0, index));
-								Log.d(TAG, "Vulkan level: " + ref.VulkanLevel);
+								Log.debug("Vulkan level: " + ref.VulkanLevel);
 							}
 						}
 					}
@@ -118,7 +116,7 @@ public class GameActivitySetupInfo implements Parcelable {
 								int VersionMinor = (ref.VulkanVersion >> 12) & 0x03ff;
 								int VersionPatch = ref.VulkanVersion & 0x0fff;
 								ref.VulkanVersionString = VersionMajor + "." + VersionMinor + "." + VersionPatch;
-								Log.d(TAG,"GetVulkanInfoFromPackage: Vulkan version: " + VersionMajor + "." + VersionMinor + "." + VersionPatch);
+								Log.debug("GetVulkanInfoFromPackage: Vulkan version: " + VersionMajor + "." + VersionMinor + "." + VersionPatch);
 							}
 						}
 					}
@@ -148,11 +146,7 @@ public class GameActivitySetupInfo implements Parcelable {
 
 		this.renderSurface = externalSurface;
 
-		if (inActivity != null)
-		{
-			Activity _AndroidActivity = inActivity.mAppActivityContext != null && inActivity.mAppActivityContext instanceof Activity ? ((Activity)inActivity.mAppActivityContext) : null;
-			this._extrasBundle = _AndroidActivity != null ? _AndroidActivity.getIntent().getExtras() : null;
-		}
+		this._extrasBundle = inActivity.mAppActivityContext != null && inActivity.mAppActivityContext instanceof Activity ? ((Activity)inActivity.mAppActivityContext).getIntent().getExtras() : null;
 
 		if (this._extrasBundle == null)
 		{
@@ -161,11 +155,11 @@ public class GameActivitySetupInfo implements Parcelable {
 		}
 		if (this.renderSurface != null) {
 			this._extrasBundle.putParcelable("renderSurface", this.renderSurface);
-			Log.d(TAG, "====> GameActivitySetupInfo(PUT) is storing renderSurface = " + this.renderSurface);
+			Log.debug("====> GameActivitySetupInfo(PUT) is storing renderSurface = " + this.renderSurface);
 		}
 
 		this.appPackageName = inActivity.getApplicationContext().getPackageName(); //inActivity.getApplicationContext().getPackageName()
-		Log.w(TAG, "====> GameActivitySetupInfo() check OBBFilename = " + OBBFilename);
+		Log.warn("====> GameActivitySetupInfo() check OBBFilename = " + OBBFilename);
 
 		this.ProjectName = this.appPackageName.substring(this.appPackageName.lastIndexOf('.') + 1);
 
@@ -183,16 +177,16 @@ public class GameActivitySetupInfo implements Parcelable {
 		}
 		catch (PackageManager.NameNotFoundException e)
 		{
-			Log.d(TAG, "Failed to load meta-data: NameNotFound: " + e.getMessage());
+			Log.debug( "Failed to load meta-data: NameNotFound: " + e.getMessage());
 		}
 		catch (NullPointerException e)
 		{
-			Log.d(TAG, "Failed to load meta-data: NullPointer: " + e.getMessage());
+			Log.debug( "Failed to load meta-data: NullPointer: " + e.getMessage());
 		}
 
 		setOBBFilename(OBBFilename);
 
-		Log.v(TAG,"====> GameActivitySetupInfo(constructor) with ProjectName = " + this.ProjectName
+		Log.verbose("====> GameActivitySetupInfo(constructor) with ProjectName = " + this.ProjectName
 				+ ", appPackageName = " + this.appPackageName
 				+ ", _OBBFilename = " + this._OBBFilename
 				+ ", packageDataInsideApkValue = " + this.packageDataInsideApkValue
@@ -241,12 +235,12 @@ public class GameActivitySetupInfo implements Parcelable {
 
 	public void DoNativeUpdates(GameActivityForMakeAAR gameActivity)
 	{
-		Log.d(TAG,"DoNativeUpdates() ExternalFilesDir path: " + this.ExternalFilesDir);
+		Log.debug("DoNativeUpdates() ExternalFilesDir path: " + this.ExternalFilesDir);
 
 		gameActivity.nativeSetGlobalActivity(this.UseExternalFilesDir, this.PublicLogFiles, this.InternalFilesDir, this.ExternalFilesDir, this.packageDataInsideApkValue, this.packageResourcePath);
 		if (this.packageDataInsideApkValue)
 		{
-			Log.d(TAG,"OBB override: " + this._OBBFilename);
+			Log.debug("OBB override: " + this._OBBFilename);
 			gameActivity.nativeSetObbFilePaths(this._OBBFilename, "", "", "");
 		}
 
@@ -270,18 +264,18 @@ public class GameActivitySetupInfo implements Parcelable {
 		}
 		catch (Exception e)
 		{
-			Log.d(TAG,"Error accessing packageInfo: " + e.getMessage());
+			Log.debug("Error accessing packageInfo: " + e.getMessage());
 		}
 
 		String productName = gameActivity.getProductName();
-		Log.d(TAG, "Android version is " + android.os.Build.VERSION.RELEASE );
-		Log.d(TAG, "Android manufacturer is " + android.os.Build.MANUFACTURER );
-		Log.d(TAG, "Android model is " + android.os.Build.MODEL );
-		Log.d(TAG, "Android build number is " + android.os.Build.DISPLAY );
-		Log.d(TAG, "OS language is set to " + Language);
-		Log.d(TAG, "Product name is " + (productName.isEmpty() ? "[not set]" : productName));
-		Log.d(TAG, "Debugger attached is " + bDebuggerAttached );
-		Log.d(TAG, "Android targetSdkVersion version is " + gameActivity.targetSdkVersion );
+		Log.debug( "Android version is " + android.os.Build.VERSION.RELEASE );
+		Log.debug( "Android manufacturer is " + android.os.Build.MANUFACTURER );
+		Log.debug( "Android model is " + android.os.Build.MODEL );
+		Log.debug( "Android build number is " + android.os.Build.DISPLAY );
+		Log.debug( "OS language is set to " + Language);
+		Log.debug( "Product name is " + (productName.isEmpty() ? "[not set]" : productName));
+		Log.debug( "Debugger attached is " + bDebuggerAttached );
+		Log.debug( "Android targetSdkVersion version is " + gameActivity.targetSdkVersion );
 		
 		gameActivity.nativeSetAndroidVersionInformation(android.os.Build.VERSION.RELEASE, gameActivity.targetSdkVersion, android.os.Build.MANUFACTURER, android.os.Build.MODEL, android.os.Build.DISPLAY, Language, productName );
 
@@ -294,7 +288,7 @@ public class GameActivitySetupInfo implements Parcelable {
 		catch (Exception e)
 		{
 			// if the above failed, then, we can't use obbs
-			Log.d(TAG,"==================================> PackageInfo failure getting .obb info: " + e.getMessage());
+			Log.debug("==================================> PackageInfo failure getting .obb info: " + e.getMessage());
 		}
 
 		// tell Android that we want volume controls to change the media volume, aka music
@@ -334,12 +328,12 @@ public class GameActivitySetupInfo implements Parcelable {
 		this._extrasBundle = in.readBundle(getClass().getClassLoader());
 		if (this._extrasBundle == null)
 		{
-			Log.e(TAG,"=====> PackageInfo _extrasBundle so no renderSurface is passed!" );
+			Log.error("=====> PackageInfo _extrasBundle so no renderSurface is passed!" );
 		}
 		this.renderSurface = this._extrasBundle == null ? null : this._extrasBundle.<Surface>getParcelable("renderSurface");
-		Log.d(TAG, "====> GameActivitySetupInfo(GET) is storing renderSurface = " + this.renderSurface);
+		Log.debug( "====> GameActivitySetupInfo(GET) is storing renderSurface = " + this.renderSurface);
 
-		Log.e(TAG,"PackageInfo from _extrasBundle =====>  renderSurface = " + this.renderSurface );
+		Log.error("PackageInfo from _extrasBundle =====>  renderSurface = " + this.renderSurface );
 
 		this._bundle = in.readBundle(getClass().getClassLoader());
 		this.appPackageName = in.readString();
diff --git a/Engine/Plugins/Runtime/AndroidSingleInstanceService/Source/AndroidSingleInstanceService/Android/src/main/java/com/epicgames/makeaar/UnrealSharedInstanceService.java b/Engine/Plugins/Runtime/AndroidSingleInstanceService/Source/AndroidSingleInstanceService/Android/src/main/java/com/epicgames/makeaar/UnrealSharedInstanceService.java
--- a/Engine/Plugins/Runtime/AndroidSingleInstanceService/Source/AndroidSingleInstanceService/Android/src/main/java/com/epicgames/makeaar/UnrealSharedInstanceService.java
+++ b/Engine/Plugins/Runtime/AndroidSingleInstanceService/Source/AndroidSingleInstanceService/Android/src/main/java/com/epicgames/makeaar/UnrealSharedInstanceService.java
@@ -25,6 +25,7 @@ import com.epicgames.makeaar.EngineFactory;
 import com.epicgames.makeaar.GameActivityForMakeAAR;
 import com.epicgames.makeaar.GameActivitySetupInfo;
 import com.epicgames.unreal.SimpleContextWrapper;
+import com.epicgames.unreal.GameActivity;
 
 import com.epicgames.makeaar.UnrealMessageType;
 
@@ -61,7 +62,7 @@ public class UnrealSharedInstanceService extends Service {
 		static private Engine engineInstance;
 
 
-		public static String defaultProjectModuleName = "AndroidPackagingTest";
+		public static String defaultProjectModuleName = "UNDEFINED_PackageName";
 //		public static String defaultProjectModuleName = "Alpha_POC";
 		public static String defaultOBBFileLocation = "";//"Gallery3D"; //"Alpha_POC";
 		//	public final static String defaultProjectModuleName = "";
@@ -260,7 +261,7 @@ public class UnrealSharedInstanceService extends Service {
 							final int PropagateAlpha = enablePropagateAlpha ? 1 : 0;
 							boolean bPortrait = applicationContext.getResources().getConfiguration().orientation == Configuration.ORIENTATION_PORTRAIT;
 
-							GameActivityForMakeAAR.Get().nativeSetWindowInfo(bPortrait, DepthBufferPreference, PropagateAlpha);
+							GameActivity.Get().nativeSetWindowInfo(bPortrait, DepthBufferPreference, PropagateAlpha);
 						}
 
 						bDoResume &= engineInstance.AttachExternalRenderSurface(taskId, externalSurface, viewPos, viewSize);
@@ -316,12 +317,12 @@ public class UnrealSharedInstanceService extends Service {
 
 					//if (bDoPause)
 					{
-						engineInstance.QueuePause("UnrealMessageType.DetachExternalSurface, taskId=" + taskId);
+						engineInstance.onPause(engineInstance.getCurrentContextID(), "UnrealMessageType.DetachExternalSurface, taskId=" + taskId);
 					}
 
 					//if (bDoStop)
 					{
-						//engineInstance.QueueStop("UnrealMessageType.DetachExternalSurface, taskId=" + taskId);
+						engineInstance.onStop(engineInstance.getCurrentContextID(), "UnrealMessageType.DetachExternalSurface, taskId=" + taskId);
 					}
 
 					//engineInstance.sendConsoleCommand("t.maxfps 0.001");
@@ -376,7 +377,7 @@ public class UnrealSharedInstanceService extends Service {
 				else {
 
 					if (bCMDMainInit)
-						GameActivityForMakeAAR.Get().nativeResumeMainInit();
+						GameActivity.Get().nativeResumeMainInit();
 				}
 
 			}
@@ -447,7 +448,7 @@ public class UnrealSharedInstanceService extends Service {
 
 		if (mMessenger != null) {
 			try {
-				mMessenger.send(Message.obtain(null, UnrealMessageType.StopService.ordinal()));
+				mMessenger.send(Message.obtain(null, UnrealMessageType.DetachExternalSurface.ordinal()));
 			} catch (RemoteException e) {
 				e.printStackTrace();
 			}
@@ -478,4 +479,4 @@ public class UnrealSharedInstanceService extends Service {
 			+ ", intent=" + intent);
 		return super.onUnbind(intent);
 	}
-}
\ No newline at end of file
+}
diff --git a/Engine/Source/Runtime/Core/Private/Android/AndroidPlatformProcess.cpp b/Engine/Source/Runtime/Core/Private/Android/AndroidPlatformProcess.cpp
--- a/Engine/Source/Runtime/Core/Private/Android/AndroidPlatformProcess.cpp
+++ b/Engine/Source/Runtime/Core/Private/Android/AndroidPlatformProcess.cpp
@@ -157,7 +157,7 @@ FString FAndroidPlatformProcess::GetGameBundleId()
 	JNIEnv* JEnv = AndroidJavaEnv::GetJavaEnv();
 	if (nullptr != JEnv)
 	{
-		jclass Class = AndroidJavaEnv::FindJavaClassGlobalRef(ANDROID_GAMEACTIVITY_CLASSPATH);
+		jclass Class = AndroidJavaEnv::FindJavaClassGlobalRef(ANDROID_GAMEACTIVITY_BASE_CLASSPATH);
 		if (nullptr != Class)
 		{
 			jmethodID getAppPackageNameMethodId = JEnv->GetStaticMethodID(Class, "getAppPackageName", "()Ljava/lang/String;");
diff --git a/Engine/Source/Runtime/Launch/Private/Android/LaunchAndroid.cpp b/Engine/Source/Runtime/Launch/Private/Android/LaunchAndroid.cpp
--- a/Engine/Source/Runtime/Launch/Private/Android/LaunchAndroid.cpp
+++ b/Engine/Source/Runtime/Launch/Private/Android/LaunchAndroid.cpp
@@ -371,24 +371,12 @@ void FPlatformMisc::UnlockAndroidWindow()
 	check(FTaskGraphInterface::IsRunning());
 
 	DEVELOPER_LOG_COMMANDCB_CASE(UnlockAndroidWindow);
-	UE_LOG(LogAndroid, Log, TEXT("Unlocking android HW window during preinit."));
-	
-	bReadyToProcessEvents = true;
+	UE_LOG(LogAndroid, Log, TEXT("Unlocking android HW window during preinit. bAppIsActive_EventThread=%d"), bAppIsActive_EventThread);
 	GAndroidWindowLock_Unlock("UnlockAndroidWindow");
 	
-}
-
-void LockAndroidWindow(FString calledBy)
-{
-	check(IsInGameThread());
-	check(FTaskGraphInterface::IsRunning());
-	
-	DEVELOPER_LOG_COMMANDCB_CASE(LockAndroidWindow);
-	UE_LOG(LogAndroid, Log, TEXT("Locking android HW window during preinit. bAppIsActive_EventThread=%d"), bAppIsActive_EventThread);
+	bReadyToProcessEvents = true;
+	//bAppIsActive_EventThread = true;
 	
-	GAndroidWindowLock_Lock(*(calledBy + "->LockAndroidWindow"));
-	bReadyToProcessEvents = false;
-	//bAppIsActive_EventThread = false;
 }
 
 JNI_METHOD void Java_com_epicgames_unreal_GameActivity_nativeResumeMainInit(JNIEnv* jenv, jobject thiz)
@@ -854,7 +842,7 @@ int32 AndroidMain(struct android_app* state)
 		FAppEventManager::GetInstance()->Tick();
 		if (!FAppEventManager::GetInstance()->IsGamePaused()
 #if USE_ANDROID_STANDALONE
-			&& FAppEventManager::GetInstance()->IsGameInFocus()
+		//	&& FAppEventManager::GetInstance()->IsGameInFocus()
 #endif
 			)
 		{
@@ -1045,10 +1033,20 @@ static void* AndroidEventThreadWorker( void* param )
 
 	TheChoreographer.SetupChoreographer();
 
+	
+
 	// window is initially invalid/locked.
-	UE_LOG(LogAndroid, Log, TEXT("event thread, Initial HW window lock."));
+	extern void* GAndroidWindowOverride;
+	UE_LOG(LogAndroid, Log, TEXT("AndroidEventThreadWorker, Initial HW window lock. GAndroidWindowOverride=%p"), GAndroidWindowOverride);
+#if USE_ANDROID_STANDALONE
+	if (GAndroidWindowOverride != nullptr)
+#endif
+	{
+		GAndroidWindowLock_Lock("AndroidEventThreadWorker");
+	}
+
 	DEVELOPER_LOG_COMMANDCB_CASE(AndroidEventThreadWorker_BeforeWhile);
-	GAndroidWindowLock_Lock("AndroidEventThreadWorker");
+
 
 	//continue to process events until the engine is shutting down
 	while (!IsEngineExitRequested())
@@ -1062,6 +1060,9 @@ static void* AndroidEventThreadWorker( void* param )
 	DEVELOPER_LOG_COMMANDCB_CASE(AndroidEventThreadWorker_AfterWhile);
 
 	GAndroidWindowLock_Unlock("AndroidEventThreadWorker");
+#if USE_ANDROID_STANDALONE
+	bReadyToProcessEvents = true;
+#endif
 
 	UE_LOG(LogAndroid, Log, TEXT("AndroidEventThreadWorker->Exiting"));
 	STANDALONE_DEBUG_LOGf(LogAndroid, TEXT("Exiting AndroidEventThreadWorker"));
@@ -1574,12 +1575,17 @@ static void ActivateApp_EventThread()
 	{
 		// Seems this can occur.
 		STANDALONE_DEBUG_LOGf(LogAndroid, TEXT("WARNING -- event thread, activate app, ALREADY have HW window lock according to bAppIsActive_EventThread. bReadyToProcessEvents=%d, bAppIsActive_EventThread=%d"), bReadyToProcessEvents, bAppIsActive_EventThread);
+#if !USE_ANDROID_STANDALONE
 		return;
+#endif
 	}
 
 	// Unlock window when we're ready.
 	UE_LOG(LogAndroid, Log, TEXT("event thread, activate app, unlocking HW window"));
 	GAndroidWindowLock_Unlock("ActivateApp_EventThread");
+#if USE_ANDROID_STANDALONE
+	bReadyToProcessEvents = true;
+#endif
 	// wake the GT up.
 	FAppEventManager::GetInstance()->EnqueueAppEvent(APP_EVENT_STATE_APP_ACTIVATED);
 
@@ -1617,7 +1623,9 @@ static void SuspendApp_EventThread()
 	if (!bAppIsActive_EventThread)
 	{
 		STANDALONE_DEBUG_LOGf(LogAndroid, TEXT("event thread, SuspendApp_EventThread called but aborted, GAndroidWindowLockRefCount=%d, bReadyToProcessEvents=%d, bAppIsActive_EventThread=%d, IsInGameThread=%d"), GAndroidWindowLockRefCount, bReadyToProcessEvents, bAppIsActive_EventThread, IsInGameThread());
+#if !USE_ANDROID_STANDALONE
 		return;
+#endif
 	}
 	bAppIsActive_EventThread = false;
 	// Lock the window, this prevents event thread from removing the window whilst the RHI initializes.
@@ -1682,12 +1690,20 @@ static void SuspendApp_EventThread()
 //Called from the event process thread
 static void OnAppCommandCB(struct android_app* app, int32_t cmd)
 {
+#if !USE_ANDROID_STANDALONE
 	check(IsInAndroidEventThread());
+#endif
 
 
 	// Set event thread's view of the window dimensions:
 	{
 		ANativeWindow* DimensionWindow = app->pendingWindow ? app->pendingWindow : app->window;
+#if USE_ANDROID_STANDALONE
+		if (GAndroidWindowOverride != NULL)
+		{
+			DimensionWindow = (ANativeWindow*)GAndroidWindowOverride;
+		}
+#endif
 		if (DimensionWindow)
 		{
 			FAndroidWindow::SetWindowDimensions_EventThread(DimensionWindow);
@@ -1941,6 +1957,7 @@ static void OnAppCommandCB(struct android_app* app, int32_t cmd)
 		DEVELOPER_LOG_COMMANDCB_CASE(APP_CMD_STOP);
 		UE_LOG(LogAndroid, Log, TEXT("Case APP_CMD_STOP"));
 		FAppEventManager::GetInstance()->EnqueueAppEvent(APP_EVENT_STATE_ON_STOP);
+		bShouldRestartFromInterrupt = true;
 		break;
 	case APP_CMD_DESTROY:
 		/**
@@ -2259,7 +2276,6 @@ JNI_METHOD void Java_com_epicgames_makeaar_GameActivityForMakeAAR_nativeMain(JNI
 	pthread_attr_setdetachstate(&otherAttr, PTHREAD_CREATE_DETACHED);
 	pthread_create(&G_AndroidMainThread, &otherAttr, android_main, &dummy_state);
 
-
 	STANDALONE_DEBUG_LOGf(LogAndroid, TEXT("Created main thread"));
 }
 
@@ -2380,14 +2396,28 @@ JNI_METHOD void Java_com_epicgames_makeaar_GameActivityForMakeAAR_nativeAppComma
 {
 	dummy_state.window = (ANativeWindow*)GAndroidWindowOverride;
 	dummy_state.pendingWindow = (ANativeWindow*)GAndroidWindowOverride;
+	int useGameThreadValue = CVarUseGameThreadForCommands.GetValueOnAnyThread();
+	bool bForceCommandToGameThread = useGameThreadValue != 0;
+	STANDALONE_DEBUG_LOGf(LogAndroid, TEXT("nativeAppCommand -> useGameThreadValue=%d, FTaskGraphInterface::IsRunning()=%d, cmd: [%u],"), useGameThreadValue, FTaskGraphInterface::IsRunning(), cmd);
 
-	bool bForceCommandToGameThread = CVarUseGameThreadForCommands.GetValueOnAnyThread() != 0;
+	if (EventThreadID == 0)
+	{
+		STANDALONE_DEBUG_LOGf(LogAndroid, TEXT("waiting for EventThreadID != 0"));
+		while (EventThreadID == 0)
+		{
+			FPlatformProcess::Sleep(0.001f);
+			FPlatformMisc::MemoryBarrier();
+		};
+	}
 
-	if (bForceCommandToGameThread && FTaskGraphInterface::IsRunning())
+	if (useGameThreadValue == 0)
+	{
+		STANDALONE_DEBUG_LOGf(LogAndroid, TEXT("OnAppCommandCB(default) cmd: %u, tid = %d, FTaskGraphInterface::IsRunning()=%d"), cmd, gettid(), FTaskGraphInterface::IsRunning());
+		OnAppCommandCB(&dummy_state, cmd);
+	}
+	else if (bForceCommandToGameThread && FTaskGraphInterface::IsRunning())
 	{
 		STANDALONE_DEBUG_LOGf(LogAndroid, TEXT("OnAppCommandCB(Queue on GameThread) cmd: %u, tid = %d, FTaskGraphInterface::IsRunning()=%d"), cmd, gettid(), FTaskGraphInterface::IsRunning());
-
-
 		FGraphEventRef nativeAppCommandCalled = FFunctionGraphTask::CreateAndDispatchWhenReady([&]()
 			{
 				OnAppCommandCB(&dummy_state, cmd);
@@ -2396,7 +2426,6 @@ JNI_METHOD void Java_com_epicgames_makeaar_GameActivityForMakeAAR_nativeAppComma
 	else
 	{
 		STANDALONE_DEBUG_LOGf(LogAndroid, TEXT("OnAppCommandCB(immediate) cmd: %d, tid = %d, FTaskGraphInterface::IsRunning()=%d"), (int)cmd, gettid(), FTaskGraphInterface::IsRunning());
-
 	    if (GEngine != NULL)
 	    {
 		    // Run on game thread to avoid race condition
diff --git a/Engine/Source/Runtime/OpenGLDrv/Private/Android/AndroidEGL.cpp b/Engine/Source/Runtime/OpenGLDrv/Private/Android/AndroidEGL.cpp
--- a/Engine/Source/Runtime/OpenGLDrv/Private/Android/AndroidEGL.cpp
+++ b/Engine/Source/Runtime/OpenGLDrv/Private/Android/AndroidEGL.cpp
@@ -1317,7 +1317,9 @@ void SetSharedContextGameCommand(TSharedPtr<FEvent, ESPMode::ThreadSafe> GTBlock
 extern bool IsInAndroidEventThread();
 void BlockRendering()
 {
+#if !USE_ANDROID_STANDALONE
 	check(IsInAndroidEventThread());
+#endif
 	check(GIsRHIInitialized);
 
 	UE_LOG(LogAndroid, Log, TEXT("Blocking renderer on invalid window."));
