Android SingleInstanceService fixes for crash and activity switching.Also updated ASISStub and output ASISProject
#android

diff --git a/Engine/Build/Android/Java/src/com/epicgames/unreal/GameActivity.java.template b/Engine/Build/Android/Java/src/com/epicgames/unreal/GameActivity.java.template
--- a/Engine/Build/Android/Java/src/com/epicgames/unreal/GameActivity.java.template
+++ b/Engine/Build/Android/Java/src/com/epicgames/unreal/GameActivity.java.template
@@ -559,6 +559,11 @@ public class GameActivity extends $${gameActivitySuperClass}$$ implements Surfac
 		return _gameActivityInstance;
 	}
 
+	public static Boolean isValidGameActivity()
+	{
+		return _gameActivityInstance != null;
+	}
+
 	public static SimpleContextWrapper GetCurrentActivityContext()
 	{
 		return _gameActivityInstance.activityContext;
@@ -3480,7 +3485,7 @@ public class GameActivity extends $${gameActivitySuperClass}$$ implements Surfac
 			nativeResumeMainInit();
 			InitCompletedOK = true;
 		}
-		else
+		else if (!isStandalone)
 		{
 			nativeOnInitialDownloadStarted();
 
@@ -3514,7 +3519,11 @@ public class GameActivity extends $${gameActivitySuperClass}$$ implements Surfac
 	@Override
 	public void onResume()
 	{
-		super.onResume();
+		if (!isStandalone)
+		{
+			super.onResume();
+		}
+
 		bIsActivityPaused.set(false);
 
 		// start memory reporter timer (runs every "memstatLoggingInterval" seconds)
@@ -3713,14 +3722,20 @@ public class GameActivity extends $${gameActivitySuperClass}$$ implements Surfac
 			});
 		}
 //$${gameActivityOnPauseAdditions}$$
-		super.onPause();
+		if (!isStandalone)
+		{
+			super.onPause();
+		}
 		Log.debug("==============> GameActive.onPause complete!");
 	}
 
 	@Override
 	public void onRestart()
 	{
-		super.onRestart();
+		if (!isStandalone)
+		{
+			super.onRestart();
+		}
 //$${gameActivityOnRestartAdditions}$$
 	}
 
@@ -3987,6 +4002,11 @@ public class GameActivity extends $${gameActivitySuperClass}$$ implements Surfac
 
 	public void AndroidThunkJava_KeepScreenOn(boolean Enable)
 	{
+		if (isStandalone)
+		{
+			return;
+		}
+
 		bKeepScreenOn = Enable;
 		if (Enable)
 		{
@@ -4021,6 +4041,11 @@ public class GameActivity extends $${gameActivitySuperClass}$$ implements Surfac
 
 	private void SetDisableScreenCaptureInternal(boolean bDisable)
 	{
+		if (isStandalone)
+		{
+			return;
+		}
+
 		if (bDisable)
 		{
 			activityContext.runOnUiThread(new Runnable()
@@ -5610,16 +5635,19 @@ public class GameActivity extends $${gameActivitySuperClass}$$ implements Surfac
 					final int modeid = modes[i].getModeId();
 					if(currentmodeid != modeid)
 					{
-						activityContext.runOnUiThread(new Runnable()
+						if (!isStandalone)
 						{
-							public void run()
+							activityContext.runOnUiThread(new Runnable()
 							{
-								Window w = activityContext.getWindow();
-								WindowManager.LayoutParams l = w.getAttributes();
-								l.preferredDisplayModeId = modeid;
-								w.setAttributes(l);
-							}
-						});
+								public void run()
+								{
+									Window w = activityContext.getWindow();
+									WindowManager.LayoutParams l = w.getAttributes();
+									l.preferredDisplayModeId = modeid;
+									w.setAttributes(l);
+								}
+							});
+						}
 						Log.debug("Found mode " + modeid + " for native refresh rate "+RefreshRate);
 					}
 					return true;
@@ -6115,22 +6143,25 @@ public class GameActivity extends $${gameActivitySuperClass}$$ implements Surfac
 					Log.warn("SustainedPerformanceMode is unsupported on this device.");
 				}
 			}
-			activityContext.runOnUiThread(new Runnable()
+			if (!isStandalone)
 			{
-				public void run()
+				activityContext.runOnUiThread(new Runnable()
 				{
-					try
+					public void run()
 					{
-						android.view.Window ActivityWindow = activityContext.getWindow();
-						Method m = android.view.Window.class.getMethod("setSustainedPerformanceMode",Boolean.TYPE);
-						m.invoke(ActivityWindow, bEnable);
-					}
-					catch (Exception e)
-					{
-						Log.debug("SetSustainedPerformanceMode: failed "+ e.getMessage());
+						try
+						{
+							android.view.Window ActivityWindow = activityContext.getWindow();
+							Method m = android.view.Window.class.getMethod("setSustainedPerformanceMode",Boolean.TYPE);
+							m.invoke(ActivityWindow, bEnable);
+						}
+						catch (Exception e)
+						{
+							Log.debug("SetSustainedPerformanceMode: failed "+ e.getMessage());
+						}
 					}
-				}
-			});
+				});
+			}
 		}
 		else
 		{
diff --git a/Engine/Build/Android/Java/src/com/epicgames/unreal/SimpleContextWrapper.java b/Engine/Build/Android/Java/src/com/epicgames/unreal/SimpleContextWrapper.java
--- a/Engine/Build/Android/Java/src/com/epicgames/unreal/SimpleContextWrapper.java
+++ b/Engine/Build/Android/Java/src/com/epicgames/unreal/SimpleContextWrapper.java
@@ -4,6 +4,7 @@ import android.app.Activity;
 import android.content.Context;
 import android.content.ContextWrapper;
 import android.os.Handler;
+import android.os.Looper;
 import android.view.View;
 import android.view.ViewGroup;
 import android.view.Window;
@@ -252,16 +253,25 @@ public final class SimpleContextWrapper extends ContextWrapper {
 	}
 
 
-	public final static void runOnMainThread (Runnable action){
+	public final void runOnMainThread (Runnable action)
+	{
+		//if (threadPoolExecutor == null) {
+		//	initThreadPoolExtractor();
+		//}
+		//threadPoolExecutor.execute(action);
 
-		if (threadPoolExecutor == null) {
-			initThreadPoolExtractor();
+		if (mHandler == null)
+		{
+			getMainExecutor().execute(action);
+		}
+		else
+		{
+			mHandler.post(action);
 		}
-		threadPoolExecutor.execute(action);
 
-//		action.run();
+		//action.run();
 
-//		new Handler(Looper.getMainLooper()).post(action);
+		//mHandler.post(action);
 	}
 
 }
\ No newline at end of file
diff --git a/Engine/Plugins/Runtime/AndroidSingleInstanceService/Source/AndroidSingleInstanceService/Android/ASISStub/src/main/java/com/epicgames/ASISStub/UseServiceActivity.java b/Engine/Plugins/Runtime/AndroidSingleInstanceService/Source/AndroidSingleInstanceService/Android/ASISStub/src/main/java/com/epicgames/ASISStub/UseServiceActivity.java
--- a/Engine/Plugins/Runtime/AndroidSingleInstanceService/Source/AndroidSingleInstanceService/Android/ASISStub/src/main/java/com/epicgames/ASISStub/UseServiceActivity.java
+++ b/Engine/Plugins/Runtime/AndroidSingleInstanceService/Source/AndroidSingleInstanceService/Android/ASISStub/src/main/java/com/epicgames/ASISStub/UseServiceActivity.java
@@ -11,12 +11,14 @@ import android.content.ComponentName;
 import android.content.Context;
 import android.content.ContextWrapper;
 import android.content.ServiceConnection;
+import android.content.res.Configuration;
 import android.graphics.SurfaceTexture;
 import android.os.Build;
 import android.os.Bundle;
 import android.os.Debug;
 import android.os.Handler;
 import android.os.IBinder;
+import android.os.Looper;
 import android.os.Message;
 import android.os.Messenger;
 import android.os.RemoteException;
@@ -30,11 +32,16 @@ import android.widget.Button;
 
 import com.epicgames.makeaar.UnrealMessageType;
 
+import java.lang.ref.WeakReference;
 import java.util.concurrent.atomic.AtomicReference;
 
 
 public class UseServiceActivity extends Activity implements View.OnClickListener
 {
+
+	Messenger mMessenger;
+//	private static WeakReference<View.OnClickListener> mCallbackRef = new WeakReference<>(IEventCallback.NO_OP);
+
 	private static final Handler.Callback mConnectionStateCallback = (Message msg) -> {
 		final String TAG = "UE-ServiceCallback";
 
@@ -70,7 +77,7 @@ public class UseServiceActivity extends Activity implements View.OnClickListener
 
 		int[] surfacePosition = new int[2];
 
-		TextureView prevTextureView = null;
+		TextureView activeTextureView = null;
 
 		final int mTaskID;
 
@@ -129,13 +136,15 @@ public class UseServiceActivity extends Activity implements View.OnClickListener
 			//textureView_.setOnTouchListener(this);
 			textureView_.setSurfaceTextureListener(this);
 
+			Log.d(TAG, "SetTextureView: " + textureView_ + ", isAvailable=" + textureView_.isAvailable() + ", isAttachedToWindow=" + textureView_.isAttachedToWindow());
+
 			if (textureView_.isAvailable() && textureView_.isAttachedToWindow())
 			{
 				textureView_.getLocationOnScreen(surfacePosition);
 
 				onSurfaceTextureAvailable(textureView_.getSurfaceTexture(), textureView_.getWidth(), textureView_.getHeight());
 			}
-			prevTextureView = textureView_;
+			activeTextureView = textureView_;
 		}
 
 		public boolean onTouch(View v, MotionEvent event) {
@@ -245,9 +254,9 @@ public class UseServiceActivity extends Activity implements View.OnClickListener
 //				//mServiceConnection.setValue( null );
 //			}
 
-			if (prevTextureView != null) {
-				prevTextureView.setSurfaceTextureListener(null);
-				prevTextureView = null;
+			if (activeTextureView != null) {
+				activeTextureView.setSurfaceTextureListener(null);
+				activeTextureView = null;
 			}
 
 			if (externalSurfaceData.get() != null)
@@ -338,8 +347,8 @@ public class UseServiceActivity extends Activity implements View.OnClickListener
 				externalSurfaceData.set(null);
 			}
 
-			if (prevTextureView != null) {
-				prevTextureView.setSurfaceTextureListener(null);
+			if (activeTextureView != null) {
+				activeTextureView.setSurfaceTextureListener(null);
 			}
 			return true;
 		}
@@ -480,15 +489,18 @@ public class UseServiceActivity extends Activity implements View.OnClickListener
 
 	final static String TAG = "UE-UseServiceActivity";
 
-	final static String obbModuleName = "AndroidPackagingTest";//""ALPHA_POC";
-	final static String servicePackageName = "com.epicgames.AndroidPackagingTest";
+	final static String obbModuleName = "Lyra";//""ALPHA_POC";
+//	final static String obbModuleName = BuildConfig.ASISModuleName; //"Lyra";//""ALPHA_POC";
+//	final static String servicePackageName = "";
+//	final static String servicePackageName = BuildConfig.ASISPackageName;
+	final static String servicePackageName = "com.epicgames.Lyra";
 	final static String obbFileLocation = ""; //"/data/local/tmp/3d/AndroidPackagingTest.main.obb.png";
+//	final static String obbFileLocation = BuildConfig.ASISOBBPath; //"/data/local/tmp/3d/AndroidPackagingTest.main.obb.png";
 
-	final static String commandLineArgs = "-nosound -launchandroidflags=0 -dpcvars='Android.UseGameThread=1'";
+	final static String commandLineArgs = "-nosound";
+//	final static String commandLineArgs = "-nosound -launchandroidflags=0 -dpcvars='Android.UseGameThreadForNativeCommands=0'";
 	final static boolean enablePropagateAlpha = true;
-
-
-
+	
 	private Button activateView1Button, activateView2Button, activateView3Button, restartActivityButton, themeChangeButton, resumeButton;
 	private Button pauseButton, stopButton, bindServiceButton, unbindServiceButton, memReportButton;
 
@@ -506,8 +518,16 @@ public class UseServiceActivity extends Activity implements View.OnClickListener
 			case 3:
 				textureView_ = (TextureView) findViewById( R.id.renderView2 );
 				break;
+			default:
+				assert false;
 		}
 
+//		try {
+//			GetServiceConnection("ActivateTextureView cleanup").detachSurfaceFromService(getTaskId(), null);
+//		} catch (Exception e) {
+//			e.printStackTrace();
+//		}
+
 		GetServiceConnection("ActivateTextureView").SetTextureView(textureView_);
 	}
 
@@ -574,7 +594,7 @@ public class UseServiceActivity extends Activity implements View.OnClickListener
 		super.onCreate( savedInstanceState );
 		AppCompatDelegate.setDefaultNightMode(next_theme);
 		getTheme().applyStyle(R.style.Theme_UnrealEngine_NoActionBar, true);
-
+		
 		setContentView( R.layout.interface_use_service_activity );
 
 
@@ -716,6 +736,11 @@ public class UseServiceActivity extends Activity implements View.OnClickListener
 
 	}
 
+	@Override
+	public void onConfigurationChanged(@NonNull Configuration newConfig) {
+		super.onConfigurationChanged(newConfig);
+	}
+
 	void handleThemeChange()
 	{
 		//GetServiceConnection().bindToUnrealInstanceService("handleThemeChange");
diff --git a/Engine/Plugins/Runtime/AndroidSingleInstanceService/Source/AndroidSingleInstanceService/Android/ASISStub/src/main/res/layout/interface_use_service_activity.xml b/Engine/Plugins/Runtime/AndroidSingleInstanceService/Source/AndroidSingleInstanceService/Android/ASISStub/src/main/res/layout/interface_use_service_activity.xml
--- a/Engine/Plugins/Runtime/AndroidSingleInstanceService/Source/AndroidSingleInstanceService/Android/ASISStub/src/main/res/layout/interface_use_service_activity.xml
+++ b/Engine/Plugins/Runtime/AndroidSingleInstanceService/Source/AndroidSingleInstanceService/Android/ASISStub/src/main/res/layout/interface_use_service_activity.xml
@@ -101,7 +101,7 @@
 			android:layout_width="wrap_content"
 			android:layout_height="wrap_content"
 			android:layout_weight="1"
-			android:minWidth="48dp"
+			android:minWidth="32dp"
 			android:onClick="onClick"
 			android:text="RestartActivity"
 			android:visibility="visible" />
@@ -128,6 +128,7 @@
 
 	<HorizontalScrollView
 		android:id="@+id/scrollViewMiddle"
+		android:contentDescription="Render Biews"
 		android:layout_width="wrap_content"
 		android:layout_height="wrap_content"
 		android:layout_marginTop="100dp"
@@ -168,8 +169,7 @@
 				android:id="@+id/renderView2"
 				android:layout_width="800dp"
 				android:layout_height="match_parent"
-				android:layout_weight="1"
-				android:background="@android:drawable/checkbox_on_background" />
+				android:layout_weight="1" />
 
 		</LinearLayout>
 	</HorizontalScrollView>
diff --git a/Engine/Plugins/Runtime/AndroidSingleInstanceService/Source/AndroidSingleInstanceService/Android/src/main/java/com/epicgames/makeaar/Engine.java b/Engine/Plugins/Runtime/AndroidSingleInstanceService/Source/AndroidSingleInstanceService/Android/src/main/java/com/epicgames/makeaar/Engine.java
--- a/Engine/Plugins/Runtime/AndroidSingleInstanceService/Source/AndroidSingleInstanceService/Android/src/main/java/com/epicgames/makeaar/Engine.java
+++ b/Engine/Plugins/Runtime/AndroidSingleInstanceService/Source/AndroidSingleInstanceService/Android/src/main/java/com/epicgames/makeaar/Engine.java
@@ -36,20 +36,24 @@ import java.lang.ref.WeakReference;
 import java.util.Arrays;
 import java.util.Objects;
 import java.util.Queue;
+import java.lang.Math;
 
 import static java.lang.System.loadLibrary;
 
 public class Engine
 {
-	public static int EVENTTYPE_INIT = 0;
-	public static int EVENTTYPE_POST_ENGINE_INIT = 1;
-	public static int EVENTTYPE_ENGINELOOP_INIT_COMPLETE = 2;
-	public static int EVENTTYPE_FRAME_BEGIN = 3;
-	public static int EVENTTYPE_FRAME_END = 4;
-	public static int EVENTTYPE_PRE_LOAD_MAP = 5;
-	public static int EVENTTYPE_POST_LOAD_MAP = 6;
-	public static int EVENTTYPE_ACTION = 7;
-	public static int EVENTTYPE_ACTIVITY_HAS_CHANGED = 8;
+
+	public static final int EVENTTYPE_INIT = 0;
+	public static final int EVENTTYPE_POST_ENGINE_INIT = 1;				// engine preInit and postInit are complete and all config files have been loaded and command lines processed.
+	public static final int EVENTTYPE_ENGINELOOP_INIT_COMPLETE = 2;		//  indicated the engine is ready to start rendering and the RHI will start creating resources and processing app events. This is where external apps can do any special handling and expect the engine to be in a ready state.
+	public static final int EVENTTYPE_FRAME_BEGIN = 3;
+	public static final int EVENTTYPE_FRAME_END = 4;
+	public static final int EVENTTYPE_PRE_LOAD_MAP = 5;					// load pak files is processed before load map is handled such as L_Load
+	public static final int EVENTTYPE_POST_LOAD_MAP = 6;				// map is loaded, this on first use is followed by EVENTTYPE_ENGINELOOP_INIT_COMPLETE to indicate the engine is about to prepare to render it's first frame.
+	public static final int EVENTTYPE_ACTION = 7;
+	public static final int EVENTTYPE_ACTIVITY_HAS_CHANGED = 8;         //activity change was recognized as part of engine init, this can be used to handle language changes and other activity related changes on app side. Activity should not be changed after this until and unless view has lost focus or onPause is issued.
+	public static final int EVENTTYPE_ENGINELOOP_SUSPENDED = 9;         //engine is suspeneded and put in an idle state.
+	
 
 	private static final String TAG = "UE_ASIS_Engine";
 
@@ -79,7 +83,7 @@ public class Engine
 		void Detach();
 	}
 
-	public static class ViewInterface<T> implements IViewInterface, SurfaceHolder.Callback, TextureView.SurfaceTextureListener, ViewTreeObserver.OnWindowFocusChangeListener//, ViewTreeObserver.OnWindowVisibilityChangeListener
+	public static class ViewInterface<T> implements IViewInterface, SurfaceHolder.Callback, TextureView.SurfaceTextureListener, ViewTreeObserver.OnWindowFocusChangeListener, ViewTreeObserver.OnWindowAttachListener//, ViewTreeObserver.OnWindowVisibilityChangeListener
 	{
 		private static final String TAG = "UE_ASIS_ViewInterface";
 		private T view;
@@ -91,7 +95,7 @@ public class Engine
 			this.view = inView;
 
 			((View)inView).getViewTreeObserver().addOnWindowFocusChangeListener(this);
-			//((View)inView).getViewTreeObserver().addOnWindowVisibilityChangeListener(this);
+			((View)inView).getViewTreeObserver().addOnWindowAttachListener(this);
 		}
 
 		public void onWindowVisibilityChanged(int visibility)
@@ -120,19 +124,41 @@ public class Engine
 
 			if (!hasFocus)
 			{
-				if ( !CheckEvent(EVENTTYPE_INIT) ) {
-					QueuePause("ViewInterface::handleViewFocusChange is false and EVENTTYPE_INIT has not been called, hasFocus=" + hasFocus);
-				}
-				else
-				{
-					QueuePause("ViewInterface::handleViewFocusChange, hasFocus=" + hasFocus);
-				}
+				QueuePause("ViewInterface::handleViewFocusChange, hasFocus=" + hasFocus);
 			}
 			else
 			{
 				QueueResume("ViewInterface::handleViewFocusChange, hasFocus=" + hasFocus + ", maskEngineEvents=" + maskEngineEvents);
 			}
 		}
+		@Override
+		public void onWindowAttached() {
+			Engine engine = Get();
+			Log.i(TAG, "** ViewInterface::onWindowAttached"
+				//+ ", bGameActivityNativeMainIsDone=" + bGameActivityNativeMainIsDone
+				+ ", bCalledStart:" + engine.bCalledStart
+				+ ", startCounter:" + startCounter
+				+ ", bResuming:" + bResuming
+				+ ", onPauseCounter:" + onPauseCounter
+				+ ", renderSurfaceValid:" + engine.renderSurfaceValid
+				+ ", renderSurface:" + engine.renderSurface
+				+ ", pendingChangedSurface:" + engine.pendingChangedSurface
+			);
+		}
+		@Override
+		public void onWindowDetached() {
+			Engine engine = Get();
+			Log.i(TAG, "** ViewInterface::onWindowDetached"
+				//+ ", bGameActivityNativeMainIsDone=" + bGameActivityNativeMainIsDone
+				+ ", bCalledStart:" + engine.bCalledStart
+				+ ", startCounter:" + startCounter
+				+ ", bResuming:" + bResuming
+				+ ", onPauseCounter:" + onPauseCounter
+				+ ", renderSurfaceValid:" + engine.renderSurfaceValid
+				+ ", renderSurface:" + engine.renderSurface
+				+ ", pendingChangedSurface:" + engine.pendingChangedSurface
+			);
+		}
 
 		public boolean canAcceptType(Class<?> candidate) {
 			return type.isAssignableFrom(candidate);
@@ -166,6 +192,10 @@ public class Engine
 					surfaceView.getHolder().setFormat(PixelFormat.TRANSPARENT);
 				}
 				surfaceView.getHolder().addCallback(this);
+				if (surfaceView.isAttachedToWindow())
+				{
+					surfaceCreated(surfaceView.getHolder());
+				}
 			}
 			else if (view instanceof TextureView)
 			{
@@ -173,6 +203,12 @@ public class Engine
 				textureView.setAlpha(1.0f);
 				textureView.setOpaque(!enablePropagateAlpha);
 				textureView.setSurfaceTextureListener(this);
+				if (textureView.isAttachedToWindow())
+				{
+					if (textureView.getSurfaceTexture() != null) {
+						onSurfaceTextureAvailable(textureView.getSurfaceTexture(), textureView.getWidth(), textureView.getHeight());
+					}
+				}
 			}
 		}
 
@@ -182,24 +218,20 @@ public class Engine
 			int[] surfacePosition = new int[] {holder.getSurfaceFrame().left, holder.getSurfaceFrame().top};
 			int[] surfaceSize = new int[]{holder.getSurfaceFrame().width(), holder.getSurfaceFrame().height()};
 
-			Log.i(TAG, "surfaceCreated, surfacePosition = " + Arrays.toString(surfacePosition) + ", surfaceSize = " + Arrays.toString(surfaceSize) + ", holder.getSurface().isValid()= " + holder.getSurface().isValid());
+			Log.v(TAG, "surfaceCreated, surfacePosition = " + Arrays.toString(surfacePosition) + ", surfaceSize = " + Arrays.toString(surfaceSize) + ", holder.getSurface().isValid()= " + holder.getSurface().isValid());
 
 			Engine engine = Get();
 
-			synchronized (syncLock) {
+			synchronized (syncLock) {				
 				engine.renderSurfaceViewPosition = surfacePosition;
 				engine.renderSurfaceViewSize = surfaceSize;
-				if (holder.getSurface().isValid() )//&& engine.renderSurfaceValid && engine.renderSurface != null)
-				{
-					engine.releaseVolatileResources("surfaceCreated and holder has valid surface!");
-				}
 				engine.pendingChangedSurface = holder.getSurface();
 			}
 		}
 
 		@Override
 		public void surfaceChanged(SurfaceHolder holder, int format, int width, int height) {
-			Log.i(TAG, "surfaceChanged, format = " + format + ", width = " + width + ", bResuming=" + bResuming + ", onPauseCounter=" + onPauseCounter + ", renderSurfaceValid=" + Get().renderSurfaceValid + ", renderSurface=" + Get().renderSurface );
+			Log.v(TAG, "surfaceChanged, format = " + format + ", width = " + width);
 
 			int[] surfacePosition = new int[] {holder.getSurfaceFrame().left, holder.getSurfaceFrame().top};
 			int[] surfaceSize = new int[]{width, height};
@@ -207,7 +239,7 @@ public class Engine
 			if (view != null)
 			{
 				GetPosSize(surfacePosition, surfaceSize);
-				Log.i(TAG, "surfaceChanged with view, format = " + format + ", surfacePosition = " + Arrays.toString(surfacePosition) + ", surfaceSize = " + Arrays.toString(surfaceSize));
+				Log.v(TAG, "surfaceChanged with view, format = " + format + ", surfacePosition = " + Arrays.toString(surfacePosition) + ", surfaceSize = " + Arrays.toString(surfaceSize));
 			}
 
 			Engine engine = Get();
@@ -223,16 +255,7 @@ public class Engine
 				engine.renderSurfaceViewPosition = surfacePosition;
 				engine.renderSurfaceViewSize = surfaceSize;
 				engine.pendingChangedSurface = holder.getSurface();
-				bResuming |= engine.pendingChangedSurface != null;
-			}
-			// queue the resume if needed at the end of surfaceChanged...
-			if (bResuming)
-			{
-				if (!engine.bCalledStart)
-				{
-					QueueStart("surfaceChanged");
-				}
-				QueueResume("surfaceChanged" );
+				engine.QueueResume("surfaceChanged");
 			}
 		}
 
@@ -241,7 +264,7 @@ public class Engine
 
 			Engine engine = Get();
 
-			Log.i(TAG, "surfaceDestroyed, holder = " + holder
+			Log.v(TAG, "surfaceDestroyed, holder = " + holder
 				+ ", lifecycleContextID=" + lifecycleContextID
 				+ ", bResuming=" + bResuming
 				+ ", surfaceCreatedCounter=" + surfaceCreatedCounter
@@ -253,7 +276,7 @@ public class Engine
 
 			synchronized (syncLock) {
 				if (engine.renderSurfaceValid && engine.renderSurface == holder.getSurface()) {
-					Log.i(TAG, "surfaceDestroyed, doing cleanup and pausing, holder = " + holder
+					Log.v(TAG, "surfaceDestroyed, doing cleanup and pausing, holder = " + holder
 							+ ", lifecycleContextID=" + lifecycleContextID
 							+ ", bResuming=" + bResuming
 							+ ", surfaceCreatedCounter=" + surfaceCreatedCounter
@@ -262,25 +285,11 @@ public class Engine
 							+ ", renderSurface=" + engine.renderSurface
 							+ ", holder.surface=" + holder.getSurface()
 					);
-					Surface pendingSurface = holder.getSurface();
-					GetOverrideGameActivity().nativeSetSurfaceOverride(null, renderSurfaceViewPosition[0], renderSurfaceViewPosition[1]);
-					if (engine.renderSurface != null && (onPauseCounter > 0 || startCounter > 0)) {
-						GetOverrideGameActivity().nativeAppCommand(APP_CMD_TERM_WINDOW);
-					}
-
-					if (engine.renderSurfaceValid && (startCounter > 0 || onPauseCounter > 0)) {
-						bResuming = true;
-					}
-
-					engine.renderSurface = engine.pendingChangedSurface;
-					engine.renderSurfaceValid = engine.renderSurface != null;
-					engine.pendingChangedSurface = pendingSurface;
-
-					engine.onPause(lifecycleContextID, "surfaceDestroyed case 1");
 
+					engine.onPause(lifecycleContextID, "surfaceDestroyed(doing cleanup)" );
 
 				} else if (holder.getSurface() != null && engine.renderSurface != null) {
-					Log.w(TAG, "surfaceDestroyed, re-bind surface, holder = " + holder
+					Log.v(TAG, "surfaceDestroyed, re-bind surface, holder = " + holder
 							+ ", lifecycleContextID=" + lifecycleContextID
 							+ ", bResuming=" + bResuming
 							+ ", surfaceCreatedCounter=" + surfaceCreatedCounter
@@ -294,7 +303,6 @@ public class Engine
 
 					if (engine.pendingChangedSurface != null) {
 						bResuming = true;
-						
 						engine.QueueResume("surfaceDestroyed valid pendingChangedSurface");
 					}
 				}
@@ -312,7 +320,7 @@ public class Engine
 				int[] surfacePosition = new int[]{0, 0};
 				textureView.getLocationOnScreen(surfacePosition);
 				int[] surfaceSize = new int[]{width, height};
-				Log.w(TAG, "onSurfaceTextureAvailable, surfacePosition = " + Arrays.toString(surfacePosition) + ", surfaceSize = " + Arrays.toString(surfaceSize));
+				Log.v(TAG, "onSurfaceTextureAvailable, surfacePosition = " + Arrays.toString(surfacePosition) + ", surfaceSize = " + Arrays.toString(surfaceSize));
 
 				renderSurfaceViewPosition = surfacePosition;
 				renderSurfaceViewSize = surfaceSize;
@@ -337,7 +345,7 @@ public class Engine
 				int[] surfacePosition = new int[]{0, 0};
 				textureView.getLocationOnScreen(surfacePosition);
 				int[] surfaceSize = new int[]{width, height};
-				Log.i(TAG, "onSurfaceTextureSizeChanged, surfacePosition = " + Arrays.toString(surfacePosition) + ", surfaceSize = " + Arrays.toString(surfaceSize));
+				Log.v(TAG, "onSurfaceTextureSizeChanged, surfacePosition = " + Arrays.toString(surfacePosition) + ", surfaceSize = " + Arrays.toString(surfaceSize));
 
 				renderSurfaceViewSize = surfaceSize;
 				GetOverrideGameActivity().nativeSetSurfaceViewInfo(renderSurfaceViewSize[0], renderSurfaceViewSize[1]);
@@ -353,7 +361,7 @@ public class Engine
 
 			synchronized (syncLock) {
 
-				Log.i(TAG, "onSurfaceTextureDestroyed, surfaceForTextureView = " + surfaceForTextureView + ", renderSurface = " + Engine.Get().renderSurface);
+				Log.v(TAG, "onSurfaceTextureDestroyed, surfaceForTextureView = " + surfaceForTextureView + ", renderSurface = " + Engine.Get().renderSurface);
 
 				Engine.Get().releaseVolatileResources("onSurfaceTextureDestroyed called");
 				if (surfaceForTextureView != null)
@@ -370,6 +378,7 @@ public class Engine
 		public void onSurfaceTextureUpdated(SurfaceTexture surfaceTexture) {
 		}
 
+
 		public void Detach()
 		{
 			Log.i(TAG, "* Engine::Detach"
@@ -392,6 +401,7 @@ public class Engine
 			}
 
 			((View)view).getViewTreeObserver().removeOnWindowFocusChangeListener(this);
+            ((View)view).getViewTreeObserver().removeOnWindowAttachListener(this);
 			//((View)view).getViewTreeObserver().removeOnWindowVisibilityChangeListener(this);
 
 			if (surfaceForTextureView != null)
@@ -449,7 +459,42 @@ public class Engine
 		return (GameActivityForMakeAAR)GameActivity.Get();
 	}
 
-	void DoGameActivityNativeMain(String projectModule, String reasonString)
+	private static GameActivityForMakeAAR GetOverrideGameActivityWhenInitComplete()
+	{
+		Log.w(TAG, "GetOverrideGameActivityWhenInitComplete called, maskEngineEvents=" + Long.toBinaryString(maskEngineEvents));
+		return (GameActivityForMakeAAR)GameActivity.Get();
+
+	//assert(GetCurrentActivityContext() != null);
+	//synchronized(GetCurrentActivityContext()) {
+
+	//	try
+	//	{
+	//		while (!CheckEvent(EVENTTYPE_POST_ENGINE_INIT))
+	//		{
+	//			Log.w(TAG, "GetOverrideGameActivityWhenInitComplete waiting for EVENTTYPE_POST_ENGINE_INIT!, maskEngineEvents=" + Long.toBinaryString(maskEngineEvents));
+	//			GetCurrentActivityContext().wait(10);
+	//			if (CheckEvent(EVENTTYPE_ENGINELOOP_SUSPENDED))
+	//			{
+	//				Log.e(TAG, "GetOverrideGameActivityWhenInitComplete is suspended!, maskEngineEvents=" + Long.toBinaryString(maskEngineEvents));
+
+	//				break;
+	//			}
+	//		}
+	//	}
+	//	catch (Exception e)
+	//	{
+	//		Log.e(TAG, "GetOverrideGameActivityWhenInitComplete EXCEPTION!, maskEngineEvents=" + Long.toBinaryString(maskEngineEvents));
+	//		e.printStackTrace();
+	//	}
+	//	finally() {
+	//		Log.i(TAG, "GetOverrideGameActivityWhenInitComplete ready!, maskEngineEvents=" + Long.toBinaryString(maskEngineEvents));
+	//		return (GameActivityForMakeAAR)GameActivity.Get();
+	//	}
+
+	//}
+}
+
+void DoGameActivityNativeMain(String projectModule, String reasonString)
 	{
 		Log.i(TAG, "DoGameActivityNativeMain is called with reasonString=" + reasonString + ", projectModule=" + projectModule);
 	
@@ -557,7 +602,7 @@ public class Engine
 				return false;
 			}
 
-			final boolean isChanging = (pendingChangedSurface != renderSurface) && pendingChangedSurface != null;
+			final boolean isChanging = (pendingChangedSurface != renderSurface) && renderSurface != null;
 
 			Log.i(TAG, "logOrder:" + logOrder + " *** Engine::ApplyPendingSurfaceChanges() - bResuming=" + bResuming
 					+ ", isChanging=" + isChanging
@@ -573,13 +618,29 @@ public class Engine
 				++surfaceCreatedCounter;
 			}
 
-			GetOverrideGameActivity().nativeSetSurfaceOverride(pendingChangedSurface, renderSurfaceViewPosition[0], renderSurfaceViewPosition[1]);
-			GetOverrideGameActivity().nativeSetSurfaceViewInfo(renderSurfaceViewSize[0], renderSurfaceViewSize[1]);
+			GetOverrideGameActivityWhenInitComplete().nativeSetSurfaceOverride(pendingChangedSurface, renderSurfaceViewPosition[0], renderSurfaceViewPosition[1]);
+			GetOverrideGameActivityWhenInitComplete().nativeSetSurfaceViewInfo(renderSurfaceViewSize[0], renderSurfaceViewSize[1]);
 			renderSurface = pendingChangedSurface;
 			renderSurfaceValid = true;
 			pendingChangedSurface = null;
-			GetOverrideGameActivity().nativeAppCommand(APP_CMD_INIT_WINDOW);
 
+/*
+			if (!CheckEvent(EVENTTYPE_ENGINELOOP_INIT_COMPLETE)) {
+				GameActivity.GetCurrentActivityContext().runOnUiThread(() ->
+					{
+						synchronized(syncLock)
+						{
+							Log.w(TAG, "FORCED -> Engine::ApplyPendingSurfaceChanges(!EVENTTYPE_ENGINELOOP_INIT_COMPLETE) - renderSurfaceValid=" + renderSurfaceValid + ", renderSurface=" + renderSurface);
+
+							//GetOverrideGameActivity().nativeAppCommand(APP_CMD_INIT_WINDOW);
+							//GetOverrideGameActivity().nativeAppCommand(APP_CMD_GAINED_FOCUS);
+							GetOverrideGameActivity().nativeResumeMainInit();
+						}
+					}
+				);
+
+			}
+*/
 
 			return isChanging;
 		}
@@ -644,7 +705,9 @@ public class Engine
 
 		if (gameActivitySetupInfo.ViewSize.length == 2)
 		{
-			renderSurfaceViewSize = gameActivitySetupInfo.ViewSize;
+			if (gameActivitySetupInfo.ViewSize[0] != 0 && gameActivitySetupInfo.ViewSize[1] != 0) {
+				renderSurfaceViewSize = gameActivitySetupInfo.ViewSize;
+			}
 		}
 
 		if (bAcitivityNeedsSetup)
@@ -678,6 +741,7 @@ public class Engine
 		Log.i(TAG, "* Engine::DoHousekeepingForViewChange(GameActivitySetupInfo) - Before recovery checks, bFirstInit=" + bFirstInit + ", bResuming=" + bResuming + ", renderSurfaceValid=" + renderSurfaceValid + ", renderSurface = " + renderSurface);
 
 		++logOrder;
+		++surfaceCreatedCounter;
 
 		if (renderSurface != null || surfaceCreatedCounter > 1)
 		{
@@ -868,7 +932,6 @@ public class Engine
 			return false;
 		}
 
-
 		Log.i(TAG, "** Engine::AttachExternalRenderSurface() - bFirstInit"
 			+ ", bGameActivityNativeMainIsDone=" + bGameActivityNativeMainIsDone
 			+ ", contextID:" + contextID
@@ -886,25 +949,16 @@ public class Engine
 			+ ", pendingChangedSurface:" + pendingChangedSurface
 		);
 
-
-		if (false == VerifyContextID(contextID, "Engine::AttachExternalRenderSurface()")) {
-			onStop(lifecycleContextID, "Engine::AttachExternalRenderSurface");
-			lifecycleContextID = contextID;
-		}
-
-		else if (renderSurface != null || surfaceCreatedCounter>0)
+		if (renderSurface != null || surfaceCreatedCounter>0)
 		{
-			releaseVolatileResources(contextID, "AttachExternalRenderSurface renderSurface != null || surfaceCreatedCounter>0");
-
+			ReleaseRenderSurfaces("AttachExternalRenderSurface(renderSurface != null || surfaceCreatedCounter>0)");
 		}
 
-		Log.i(TAG, "** Engine::AttachExternalRenderSurface() - CHECK 1A");
 		++logOrder;
+		++surfaceCreatedCounter;
 
-
-		//synchronized(syncLock)
-		{
-		Log.i(TAG, "** Engine::AttachExternalRenderSurface() - CHECK 1B");
+		synchronized (syncLock) {
+			Engine.Get().sendConsoleCommand("a.AllowFrameTimestamps 0");	// needed to avoid E/FrameEvents: updateAcquireFence: Did not find frame. spamming.
 			pendingChangedSurface = externalSurface;
 			if (viewSize.length==2)
 			{
@@ -968,7 +1022,7 @@ public class Engine
 			return;
 		}
 
-		GetOverrideGameActivity().GetCurrentActivityContext().runOnMainThread(() ->
+		GameActivityForMakeAAR.GetCurrentActivityContext().runOnMainThread(() ->
 				{
 					synchronized(syncLock)
 					{
@@ -992,9 +1046,9 @@ public class Engine
 			return;
 		}
 
-		GetOverrideGameActivity().GetCurrentActivityContext().runOnUiThread(() ->
+		GameActivityForMakeAAR.GetCurrentActivityContext().runOnMainThread(() ->
 			{
-				//synchronized (syncLock)
+				synchronized (syncLock)
 				{
 					Log.i(TAG, "QueueResume is calling engine.onResume with lifecycleContextID=" + lifecycleContextID + ", reasonString=" + reasonString);
 					
@@ -1012,7 +1066,7 @@ public class Engine
 			return;
 		}
 
-		GetOverrideGameActivity().GetCurrentActivityContext().runOnMainThread(new Runnable()
+		GameActivityForMakeAAR.GetCurrentActivityContext().runOnMainThread(new Runnable()
 		{
 			@Override
 			public void run()
@@ -1036,17 +1090,11 @@ public class Engine
 			return;
 		}
 
-		GetOverrideGameActivity().GetCurrentActivityContext().runOnMainThread(new Runnable()
-		{
-			@Override
-			public void run()
+		GameActivityForMakeAAR.GetCurrentActivityContext().runOnMainThread(() ->
 			{
-				synchronized(syncLock)
-				{
-					Log.i(TAG, "QueueStop is calling engine.onStop with lifecycleContextID=" + lifecycleContextID + ", reasonString=" + reasonString);
-			Engine.Get().onStop(lifecycleContextID, "QueueStop with reasonString=" + reasonString);
-				}
-	
+				synchronized (syncLock) {
+				Log.i(TAG, "QueueStop is calling engine.onStop with lifecycleContextID=" + lifecycleContextID + ", reasonString=" + reasonString);
+				Engine.Get().onStop(lifecycleContextID, "QueueStop with reasonString=" + reasonString);
 			}
 		});
 	}
@@ -1118,104 +1166,19 @@ public class Engine
 		}
 
 
-		if (!bCalledStart)
-		{
-			QueueStart("Init(context, gameActivitySetupInfo)");
-			QueueResume("Init(context, gameActivitySetupInfo) because started");
-		}
-
-		if (renderSurface != null || pendingChangedSurface != null)
-		{
-			QueueResume("Init(context, gameActivitySetupInfo) refresh surface");
-		}
+//		if (!bCalledStart)
+//		{
+//			QueueStart("Init(context, gameActivitySetupInfo)");
+//		}
+//
+//		if (renderSurface != null || pendingChangedSurface != null)
+//		{
+//			QueueResume("Init(context, gameActivitySetupInfo) refresh surface");
+//		}
 
 		return bResuming;
 	}
 
-	public boolean Init(SurfaceView view, String commandline, String projectModule, String[] mountPAKs)
-	{
-		return Init(view.getContext(), new ViewInterface<SurfaceView>(view), this.OBBFilename, commandline, projectModule, _enablePropagateAlpha, mountPAKs);
-	}
-
-	public boolean Init(TextureView view, String commandline, String projectModule, String[] mountPAKs)
-	{
-		return Init(view.getContext(), new ViewInterface<TextureView>(view), this.OBBFilename, commandline, projectModule, _enablePropagateAlpha, mountPAKs);
-	}
-
-	public boolean Init(Context inContext, IViewInterface view, String OBBFilename, String commandline, String projectModule, boolean enablePropagateAlpha, String[] mountPAKs)
-	{
-		incCreateContextID();
-
-		++logOrder;
-		Log.i(TAG, "logOrder:" + logOrder + " * Engine::Init(IViewInterface) -" + ", bLibraryLoaded=" + bLibraryLoaded + ", bResuming=" + bResuming + ", renderSurfaceValid=" + renderSurfaceValid + ", view = " + view);
-		if (renderSurfaceValid || renderSurface != null)
-		{
-			releaseVolatileResources( "Engine::Init(IViewInterface) called so remove old resources");
-		}
-		if (!bLibraryLoaded)
-		{
-			return false;
-		}
-
-		boolean bFirstInit = (GameActivity.Get() == null || renderSurface == null);
-
-
-		if (!DoHousekeepingForViewChange(inContext, view, OBBFilename, enablePropagateAlpha ))
-		{
-			return false;
-		}
-
-		bHasBeenInit = true;
-
-		if (bFirstInit)
-		{
-			Log.i(TAG, "** Engine::Init(ViewInterface) - bFirstInit"
-					+ ", bGameActivityNativeMainIsDone=" + bGameActivityNativeMainIsDone
-					+ ", bCalledStart:" + bCalledStart
-					+ ", startCounter:" + startCounter
-					+ ", bResuming:" + bResuming
-					+ ", onPauseCounter:" + onPauseCounter
-					+ ", renderSurfaceValid:" + renderSurfaceValid
-					+ ", renderSurface:" + renderSurface
-					+ ", commandline:" + commandline
-					+ ", projectModule:" + projectModule
-			);
-
-			if (commandline.isEmpty() || !commandline.startsWith("../../"))
-			{
-				commandline = "../../../" + projectModule + "/" + projectModule + ".uproject " + commandline;
-				Log.i(TAG, "** Engine::Init(ViewInterface) - commandline updated: " + commandline);
-			}
-
-			GetOverrideGameActivity().setCommandline(commandline);
-
-			if (mountPAKs.length > 0) {
-				pendingPakList = mountPAKs.clone();
-			}
-
-			// start game thread
-			if (!bGameActivityNativeMainIsDone)
-			{
-				Log.i(TAG, "** Engine::Init(IViewInterface) - calling GetOverrideGameActivity().nativeMain() with : projectModule = " + projectModule);
-				DoGameActivityNativeMain(projectModule, "Init(IViewInterface)");
-			}
-
-			AndroidThunkJava_Engine_ReceiveEvent(EVENTTYPE_ACTIVITY_HAS_CHANGED, "", 0, 0, 0.0f);
-		}
-
-
-		if (!bCalledStart)
-		{
-			QueueStart("Init(ViewInterface)");
-		}
-
-		if (pendingChangedSurface != null)
-		{
-			QueueResume("Init(ViewInterface)");
-		}
-
-		return bResuming;
-	}
 
 	public String getObbVersion(boolean bFull)
 	{
@@ -1266,8 +1229,14 @@ public class Engine
 
 	public void sendConsoleCommand(String command)
 	{
-		if (GetOverrideGameActivity() != null) {
-			GetOverrideGameActivity().nativeConsoleCommand(command);
+		if (GetOverrideGameActivity() != null) 
+		{
+			if (!CheckEvent(EVENTTYPE_ENGINELOOP_INIT_COMPLETE))
+			{
+				Log.e(TAG, "Engine::sendConsoleCommand(aborted cause not ready), maskEngineEvents=" + Long.toBinaryString(maskEngineEvents) );
+				return;
+			}
+			GetOverrideGameActivityWhenInitComplete().nativeConsoleCommand(command);
 		}
 	}
 
@@ -1305,6 +1274,12 @@ public class Engine
 
 	public void onPause(int contextID, String reasonString)
 	{
+		if (!CheckEvent(EVENTTYPE_ENGINELOOP_INIT_COMPLETE))
+		{
+			Log.e(TAG, "Engine::onPause(aborted cause not ready), maskEngineEvents=" + Long.toBinaryString(maskEngineEvents) );
+			return;
+		}
+		
 		if (contextID != 0 && lifecycleContextID != 0 && contextID != lifecycleContextID) {
 			Log.e(TAG, "[lifecycleContextID=" + lifecycleContextID + "] contextID:" + contextID + " ABNORMAL Engine::onPause(contextID) - onPause called. with contextID change = " + contextID + ", reasonString=" + reasonString);
 			return;
@@ -1317,9 +1292,9 @@ public class Engine
 
 		Log.i(TAG, "[lifecycleContextID=" + lifecycleContextID + "] logOrder:" + logOrder + " * Engine::onPause - renderSurfaceValid: " + renderSurfaceValid + ", bResuming: " + bResuming + ", bCalledStart = " + bCalledStart + ", onPauseCounter=" + onPauseCounter + ", renderSurface = " + renderSurface + ", pendingChangedSurface=" + pendingChangedSurface + ", reasonString=" + reasonString);
 
-		//if (_gameActivityInstance != null) {
-		//	GetOverrideGameActivity().onPause();
-		//}
+		if (GetOverrideGameActivity() != null) {
+			GetOverrideGameActivityWhenInitComplete().onPause();
+		}
 
 		if (onPauseCounter > 0 || bResuming) {
 			Log.w(TAG, "** Engine::onPause - Already Paused! onPauseCounter=" + onPauseCounter + ", bResuming=" + bResuming + ", reasonString=" + reasonString);
@@ -1328,24 +1303,25 @@ public class Engine
 
 		++onPauseCounter;
 		Log.i(TAG, "** Engine::onPause - DOING a full Pause!" + ", reasonString=" + reasonString);
+		GetOverrideGameActivityWhenInitComplete().nativeAppCommand(APP_CMD_PAUSE);
+		GetOverrideGameActivityWhenInitComplete().nativeAppCommand(APP_CMD_LOST_FOCUS);
 
-		if (renderSurface != null) {
-			--surfaceCreatedCounter;
-		}
+		if (!bResuming && renderSurfaceValid)
+		{
+			GetOverrideGameActivityWhenInitComplete().nativeSetSurfaceOverride(null, renderSurfaceViewPosition[0], renderSurfaceViewPosition[1]);
 
-		if (renderSurface != null || pendingChangedSurface != null) {
-			GetOverrideGameActivity().nativeSetSurfaceOverride(null, renderSurfaceViewPosition[0], renderSurfaceViewPosition[1]);
-			pendingChangedSurface = pendingChangedSurface != null ? pendingChangedSurface : renderSurface;
+			if (renderSurface != null) 
+			{
+				pendingChangedSurface = renderSurface;
+				// com.epicgames.makeaar.GameActivity.nativeAppCommand(APP_CMD_TERM_WINDOW);
+			}
+
+			bResuming |= pendingChangedSurface != null;
+	
 			renderSurface = null;
 			renderSurfaceValid = false;
-			bResuming |= pendingChangedSurface != null;
-			Log.i(TAG, "** Engine::onPause - set SurfaceOverride to null and updated render surfaces" + ", pendingChangedSurface=" + pendingChangedSurface + ", reasonString=" + reasonString);
 		}
 
-		GetOverrideGameActivity().nativeAppCommand(APP_CMD_LOST_FOCUS);
-		GetOverrideGameActivity().nativeAppCommand(APP_CMD_PAUSE);
-		//GetOverrideGameActivity().nativeAppCommand(APP_CMD_TERM_WINDOW);
-
 		//releaseVolatileResources("onPause called with reason=" + reasonString);
 	}
 
@@ -1353,12 +1329,26 @@ public class Engine
 	{
 		if (contextID != 0 && lifecycleContextID != 0 && contextID != lifecycleContextID) {
 			Log.e(TAG, "[lifecycleContextID=" + lifecycleContextID + "] contextID:" + contextID + " ABNORMAL Engine::onResume - onResume called. with contextID change = " + contextID + ", reasonString=" + reasonString);
-			return;
+			//return;
 		}
+
+		lifecycleContextID = contextID;
+		onResume(reasonString);
+	}
+
+	private void onResume(String reasonString)
+	{
+		if (!CheckEvent(EVENTTYPE_ENGINELOOP_INIT_COMPLETE))
+		{
+			Log.e(TAG, "Engine::onResume(WARNING cause not ready), maskEngineEvents=" + Long.toBinaryString(maskEngineEvents) + ", reasonString=" + reasonString );
+			//ApplyPendingSurfaceChanges();
+			//return;
+		}
+		
 		++logOrder;
-		Log.i(TAG, "[lifecycleContextID=" + lifecycleContextID + "] logOrder:" + logOrder + " * Engine::onResume - onResume called. bCalledStart was = " + bCalledStart
+		Log.i(TAG, "[lifecycleContextID=" + lifecycleContextID + "] logOrder:" + logOrder + " * Engine::onResume() - onResume called. bCalledStart was = " + bCalledStart
 			+ ", bHasBeenInit = " + bHasBeenInit
-			+ ", maskEngineEvents= " + maskEngineEvents
+			+ ", maskEngineEvents= " + Long.toBinaryString(maskEngineEvents)
 			+ ", bGameActivityNativeMainIsDone = " + bGameActivityNativeMainIsDone
 			+ ", prev startCounter = " + startCounter
 			+ ", surfaceCreatedCounter = " + surfaceCreatedCounter
@@ -1370,19 +1360,19 @@ public class Engine
 
 		if (bCalledStart == false) {
 			if (surfaceCreatedCounter > 0 && renderSurface != null) {
-				Log.w(TAG, "** Engine::onResume - FORCING onstart from onresume because bCalledStart is false! bResuming=" + bResuming + ", onPauseCounter=" + onPauseCounter + ", width = " + renderSurfaceViewSize[0] + ", height = " + renderSurfaceViewSize[1]);
+				Log.w(TAG, "** Engine::onResume() - FORCING onstart from onresume because bCalledStart is false! bResuming=" + bResuming + ", onPauseCounter=" + onPauseCounter + ", width = " + renderSurfaceViewSize[0] + ", height = " + renderSurfaceViewSize[1]);
 
-				onStart(lifecycleContextID, reasonString + "->onResume");
+				onStart(reasonString + "->onResume");
 
 			} else {
-				Log.w(TAG, "** Engine::onResume - SKIPPING because bCalledStart is false! bResuming=" + bResuming + ", onPauseCounter=" + onPauseCounter + ", width = " + renderSurfaceViewSize[0] + ", height = " + renderSurfaceViewSize[1]);
+				Log.w(TAG, "** Engine::onResume() - SKIPPING because bCalledStart is false! bResuming=" + bResuming + ", onPauseCounter=" + onPauseCounter + ", width = " + renderSurfaceViewSize[0] + ", height = " + renderSurfaceViewSize[1]);
 				return;
 			}
 		}
 
-//		if (GameActivity.Get() != null) {
-//			GameActivity.Get().onResume();
-//		}
+		if (GetOverrideGameActivity() != null && CheckEvent(EVENTTYPE_ENGINELOOP_INIT_COMPLETE)) {
+			GetOverrideGameActivity().onResume();
+		}
 
 
 		bResuming |= ApplyPendingSurfaceChanges();
@@ -1406,25 +1396,27 @@ public class Engine
 		{
 			Log.i(TAG, "** Engine::onResume - setup renderSurface because we already are in a start state! width = " + renderSurfaceViewSize[0] + ", height = " + renderSurfaceViewSize[1]);
 
-			GetOverrideGameActivity().nativeSetSurfaceOverride(renderSurface, renderSurfaceViewPosition[0], renderSurfaceViewPosition[1]);
-			GetOverrideGameActivity().nativeSetSurfaceViewInfo(renderSurfaceViewSize[0], renderSurfaceViewSize[1]);
+			GetOverrideGameActivityWhenInitComplete().nativeSetSurfaceOverride(renderSurface, renderSurfaceViewPosition[0], renderSurfaceViewPosition[1]);
+			GetOverrideGameActivityWhenInitComplete().nativeSetSurfaceViewInfo(renderSurfaceViewSize[0], renderSurfaceViewSize[1]);
 			bResuming = true;
 		}
 
 		if (bResuming && renderSurface != null) {
 			Log.i(TAG, "** Engine::onResume - DOING full Resume! width = " + renderSurfaceViewSize[0] + ", height = " + renderSurfaceViewSize[1] + ", maskEngineEvents= " + maskEngineEvents);
 
-			GetOverrideGameActivity().nativeAppCommand(APP_CMD_GAINED_FOCUS);
-			GetOverrideGameActivity().nativeResumeMainInit();
+			GetOverrideGameActivityWhenInitComplete().nativeAppCommand(APP_CMD_INIT_WINDOW);
+			GetOverrideGameActivityWhenInitComplete().nativeAppCommand(APP_CMD_GAINED_FOCUS);
+			GetOverrideGameActivityWhenInitComplete().nativeResumeMainInit();
 			bResuming = false;
 		}
 		else if (renderSurface != null)
 		{
 			Log.e(TAG, "***- Engine::onResume - UNHANDLED case for resume! renderSurfaceValid = " + renderSurfaceValid + ", renderSurface = " + renderSurface);
-			GetOverrideGameActivity().nativeSetSurfaceOverride(renderSurface, renderSurfaceViewPosition[0], renderSurfaceViewPosition[1]);
-			GetOverrideGameActivity().nativeSetSurfaceViewInfo(renderSurfaceViewSize[0], renderSurfaceViewSize[1]);
-			GetOverrideGameActivity().nativeAppCommand(APP_CMD_GAINED_FOCUS);
-			GetOverrideGameActivity().nativeResumeMainInit();
+			GetOverrideGameActivityWhenInitComplete().nativeSetSurfaceOverride(renderSurface, renderSurfaceViewPosition[0], renderSurfaceViewPosition[1]);
+			GetOverrideGameActivityWhenInitComplete().nativeSetSurfaceViewInfo(renderSurfaceViewSize[0], renderSurfaceViewSize[1]);
+			GetOverrideGameActivityWhenInitComplete().nativeAppCommand(APP_CMD_INIT_WINDOW);
+			GetOverrideGameActivityWhenInitComplete().nativeAppCommand(APP_CMD_GAINED_FOCUS);
+			GetOverrideGameActivityWhenInitComplete().nativeResumeMainInit();
 		}
 		else
 		{
@@ -1436,11 +1428,18 @@ public class Engine
 	{
 		if (contextID != 0 && lifecycleContextID != 0 && contextID != lifecycleContextID) {
 			Log.e(TAG, "[lifecycleContextID=" + lifecycleContextID + "] contextID:" + contextID + " ABNORMAL Engine::onStart(contextID) - onStart called. with contextID change = " + contextID + ", reasonString=" + reasonString);
-			return;
+			//return;
 		}
 
+		lifecycleContextID = contextID;
+		onStart(reasonString);
+	}
+
+	private void onStart(String reasonString)
+	{
 		++logOrder;
 		Log.i(TAG, "[lifecycleContextID=" + lifecycleContextID + "] logOrder:" + logOrder + " * Engine::onStart - onStart called. bCalledStart was = " + bCalledStart
+			+ "\n\t, maskEngineEvents= " + Long.toBinaryString(maskEngineEvents)
 			+ ", bHasBeenInit = " + bHasBeenInit
 			+ ", prev startCounter = " + startCounter
 			+ ", bResuming was = " + bResuming
@@ -1450,10 +1449,16 @@ public class Engine
 		);
 
 		if (!bHasBeenInit) {
+			//RegisterLifecycleEvent( Lifecycle_Start );
+			//return;
 		}
 
 		++startCounter;
-		if (!bCalledStart) {
+		if (!bCalledStart)
+		{
+			//GetOverrideGameActivity().nativeResumeMainInit();
+			//GetOverrideGameActivity().nativeAppCommand(APP_CMD_INIT_WINDOW);
+			//GetOverrideGameActivity().nativeAppCommand(APP_CMD_GAINED_FOCUS);
 			bCalledStart = true;
 			GetOverrideGameActivity().nativeAppCommand(APP_CMD_START);
 		}
@@ -1463,11 +1468,17 @@ public class Engine
 	{
 		if (contextID != 0 && lifecycleContextID != 0 && contextID != lifecycleContextID) {
 			Log.e(TAG, "[lifecycleContextID=" + lifecycleContextID + "] contextID:" + contextID + " ABNORMAL Engine::onStop(contextID) - onStop called. with contextID change = " + contextID + ", reasonString=" + reasonString);
-			return;
+			//return;
 		}
+		lifecycleContextID = contextID;
+		onStop(reasonString);
+	}
 
+	private void onStop(String reasonString)
+	{
 		++logOrder;
-		Log.i(TAG, "[lifecycleContextID=" + lifecycleContextID + "] logOrder:" + logOrder + " * Engine::onStop - onStop called. bCalledStart was = " + bCalledStart
+		Log.v(TAG, "[lifecycleContextID=" + lifecycleContextID + "] logOrder:" + logOrder + " * Engine::onStop - onStop called. bCalledStart was = " + bCalledStart
+			+ "\n\t, maskEngineEvents= " + Long.toBinaryString(maskEngineEvents)
 				+ ", reasonString=" + reasonString
 				+ ", bHasBeenInit = " + bHasBeenInit
 				+ ", prev startCounter = " + startCounter
@@ -1478,11 +1489,7 @@ public class Engine
 				+ ", renderSurface = " + renderSurface
 		);
 
-		if (onPauseCounter <= 0 || renderSurface != null) {
-			onPause(contextID, reasonString + "from onStop()");
-		}
-
-		//releaseVolatileResources("onStop() called so release all resources");
+		releaseVolatileResources("onStop withReason <- " + reasonString);
 		synchronized (syncLock) {
 			pendingChangedSurface = null;
 		}
@@ -1493,166 +1500,116 @@ public class Engine
 
 		if (bCalledStart)
 		{
-			GetOverrideGameActivity().nativeAppCommand(APP_CMD_STOP);
-			GetOverrideGameActivity().nativeAppCommand(APP_CMD_TERM_WINDOW);
+			GetOverrideGameActivityWhenInitComplete().nativeAppCommand(APP_CMD_STOP);
 
 			bCalledStart = false;
 			startCounter = 0;
 			onPauseCounter = 0;
-			lifecycleContextID = 0;
 		}
-
-
 	}
 
 	private void ReleaseRenderSurfaces(int contextID, String reasonString)
 	{
-		Log.i(TAG, "*** Engine::ReleaseRenderSurfaces() : GameActivityForMakeAAR.isValidGameActivity() = " + GameActivityForMakeAAR.isValidGameActivity()
-			+ ", bGameActivityNativeMainIsDone=" + bGameActivityNativeMainIsDone
-			+ ", contextID:" + contextID
-			+ ", lifecycleContextID:" + lifecycleContextID
-			+ ", reasonString:" + reasonString
-			+ ", bCalledStart:" + bCalledStart
-			+ ", startCounter:" + startCounter
-			+ ", bResuming:" + bResuming
-			+ ", onPauseCounter:" + onPauseCounter
-			+ ", renderSurfaceValid:" + renderSurfaceValid
-			+ ", renderSurface:" + renderSurface
-			+ ", surfaceCreatedCounter:" + surfaceCreatedCounter
-			+ ", pendingChangedSurface:" + pendingChangedSurface
-
-		);
-
-		if (GameActivityForMakeAAR.isValidGameActivity() && (onPauseCounter <= 0 || renderSurface != null || pendingChangedSurface != null)) {
-
-			onPause(contextID, reasonString + "->ReleaseRenderSurfaces");
-
 
-			//if (renderSurface != null)
-			{
-				//GetOverrideGameActivity().nativeAppCommand(APP_CMD_PAUSE);
-				//GetOverrideGameActivity().nativeAppCommand(APP_CMD_STOP);
-
-				
-
-				//if (onPauseCounter > 0) {
-				//	QueuePause("ReleaseRenderSurfaces");
-				//}
-
-				//if (startCounter > 0) {
-				//	QueueStop("ReleaseRenderSurfaces");
-				//}
-			}
-			//else if (!CheckEvent(EVENTTYPE_INIT))
-			//{
-			//	GetOverrideGameActivity().nativeAppCommand(APP_CMD_PAUSE);
-			//}
-
-			if (renderSurfaceValid && (startCounter > 0 || onPauseCounter > 0)) {
-				bResuming = true;
-			}
-				
-			if (renderSurface != null) {
-				pendingChangedSurface = renderSurface;
-				bResuming = true;
-			}
-			renderSurface = null;
-			renderSurfaceValid = false;
+		if (contextID != 0 && lifecycleContextID != 0 && contextID != lifecycleContextID) {
+			Log.e(TAG, "[lifecycleContextID=" + lifecycleContextID + "] contextID:" + contextID + " ABNORMAL Engine::releaseVolatileResources() called. with contextID change = " + contextID + ", reasonString=" + reasonString);
+			//return;
 		}
 
+		lifecycleContextID = contextID;
+		ReleaseRenderSurfaces(reasonString);
 	}
 
-	private void releaseVolatileResources(String reasonString)
+	private void ReleaseRenderSurfaces(String reasonString)
 	{
-		releaseVolatileResources(lifecycleContextID, reasonString);
+		synchronized (syncLock) {
+			if (GameActivity.isValidGameActivity()) {
+		        Log.i(TAG, "*** Engine::ReleaseRenderSurfaces() : GameActivityForMakeAAR.isValidGameActivity() = " + GameActivityForMakeAAR.isValidGameActivity()
+					+ "\n\t, maskEngineEvents= " + Long.toBinaryString(maskEngineEvents)
+			        + ", bGameActivityNativeMainIsDone=" + bGameActivityNativeMainIsDone
+			        + ", lifecycleContextID:" + lifecycleContextID
+			        + ", reasonString:" + reasonString
+			        + ", bCalledStart:" + bCalledStart
+			        + ", startCounter:" + startCounter
+			        + ", bResuming:" + bResuming
+			        + ", onPauseCounter:" + onPauseCounter
+			        + ", renderSurfaceValid:" + renderSurfaceValid
+			        + ", renderSurface:" + renderSurface
+			        + ", surfaceCreatedCounter:" + surfaceCreatedCounter
+			        + ", pendingChangedSurface:" + pendingChangedSurface
+
+		        );
+				        
+		        --surfaceCreatedCounter;
+
+		        if (renderView != null) {
+			        renderView.Detach();
+			        renderView = null;
+		        }
+
+				GetOverrideGameActivityWhenInitComplete().nativeSetSurfaceOverride(pendingChangedSurface, renderSurfaceViewPosition[0], renderSurfaceViewPosition[1]);
+		        if (renderSurface != null && (onPauseCounter > 0 || startCounter > 0)) {
+					GetOverrideGameActivityWhenInitComplete().nativeAppCommand(APP_CMD_TERM_WINDOW);
+		        }
+
+				if (renderSurfaceValid && (startCounter > 0 || onPauseCounter > 0)) {
+					bResuming = true;
+				}
+                else {
+	                bCalledStart = false;
+                }
+
+				renderSurface = pendingChangedSurface;
+				renderSurfaceValid = renderSurface != null;
+				pendingChangedSurface = null;
+			}
+		}
 	}
 
 	private void releaseVolatileResources(int contextID, String reasonString)
 	{
-		Log.i(TAG, "logOrder:" + logOrder + " * Engine::releaseVolatileResources() - bCalledStart was = " + bCalledStart + ", contextID=" + contextID + ", reasonString=" + reasonString + ", bResuming was = " + bResuming + ", onPauseCounter=" + onPauseCounter + ", renderSurfaceValid: " + renderSurfaceValid + ", renderSurface = " + renderSurface );
+		if (contextID != 0 && lifecycleContextID != 0 && contextID != lifecycleContextID) {
+			Log.e(TAG, "[lifecycleContextID=" + lifecycleContextID + "] contextID:" + contextID + " ABNORMAL Engine::releaseVolatileResources() called. with contextID change = " + contextID + ", reasonString=" + reasonString);
+			//return;
+		}
 
-		ReleaseRenderSurfaces(contextID, reasonString);
+		lifecycleContextID = contextID;
+		releaseVolatileResources(reasonString);
 	}
 
-	public void onDestroy(int contextID)
+	private void releaseVolatileResources(String reasonString)
 	{
-		onDestroy(contextID, "External");
+		Log.i(TAG, "logOrder:" + logOrder + " * Engine::releaseVolatileResources() - bCalledStart was = " + bCalledStart + ", lifecycleContextID=" + lifecycleContextID + ", reasonString=" + reasonString + ", bResuming was = " + bResuming + ", onPauseCounter=" + onPauseCounter + ", renderSurfaceValid: " + renderSurfaceValid + ", renderSurface = " + renderSurface );
+
+		ReleaseRenderSurfaces(lifecycleContextID, reasonString);
 	}
 
-	private void onDestroy(int customContextID, String reasonString)
+	public void onDestroy(int customContextID, String reasonString)
 	{
 		++logOrder;
 		Log.i(TAG, "[lifecycleContextID=" + lifecycleContextID + "] logOrder:" + logOrder + " * Engine::onDestroy() - onDestroy called. bCalledStart was = " + bCalledStart + ", bResuming was = " + bResuming + ", onPauseCounter=" + onPauseCounter + ", renderSurfaceValid: " + renderSurfaceValid + ", renderSurface = " + renderSurface );
 
-		//onPause(customContextID, "onDestroy called");
-		ReleaseRenderSurfaces(customContextID, reasonString + "->onDestroy called");
-
-		if (customContextID == 543210 || !CheckEvent(EVENTTYPE_INIT))
+		if (customContextID == 543210)
 		{
-
 			// we don't want to call APP_CMD_DESTROY unless we want to completely terminate the unreal engine instance.
-			GetOverrideGameActivity().nativeAppCommand(APP_CMD_DESTROY);
-
-		}
-
-		if (renderSurface != null) {
-			GetOverrideGameActivity().nativeAppCommand(APP_CMD_TERM_WINDOW);
-		}
-
-		if (renderView != null) {
-			renderView.Detach();
-			renderView = null;
+			GetOverrideGameActivityWhenInitComplete().nativeAppCommand(APP_CMD_DESTROY);
 		}
 
 		maskEngineEvents = 0;
 		bHasBeenInit = false;
-		bResuming = false;
-
 
 		lifecycleContextID += 25;
 		if (onPauseCounter>0)
 		{
-			Log.e(TAG, "onDestroy is called with onPauseCounter>0 which should not happen for healthy lifecycle");
+			QueueResume("onDestroy withReason <- " + reasonString);
 		}
-
-		Log.w(TAG, "*** Engine::onDestroy() -> AndroidThunkJava_RestartApplication()");
-		GetOverrideGameActivity().AndroidThunkJava_RestartApplication("");
-
-		Log.w(TAG, "*** Engine::onDestroy() set statics to defaults!");
-
-		surfaceCreatedCounter = 0;
-		activity = null;
-		//bLibraryLoaded = false;
-		renderView = null;
-
-		maskEngineEvents = 0;
-		renderSurfaceValid = false;
-		renderSurface = null;
-		pendingChangedSurface = null;
-		//renderSurfaceViewPosition[] = new int[2];
-		//renderSurfaceViewSize[] = new int[2];
-		//syncLock = new Object();
-		//private GameActivityForMakeAAR gameActivity = null; //TODO remove ref here and use Get() on GameActivity
-		//private String OBBFilename = null;
-		//private boolean bOBBinAPK = true;
-
-		bResuming = false;
-		//private boolean bAllowConsole = false;
-		bGameActivityNativeMainIsDone = false;
-
-
-	
-		Log.i(TAG, "[lifecycleContextID=" + lifecycleContextID + "] logOrder:" + logOrder + " * Engine::onDestroy() - onDestroy completed. bCalledStart was = " + bCalledStart + ", bResuming was = " + bResuming + ", onPauseCounter=" + onPauseCounter + ", renderSurfaceValid: " + renderSurfaceValid + ", renderSurface = " + renderSurface );
-
-
-	
 	}
 
 	public void onConfigurationChanged(Configuration newConfig)
 	{
 		// forward the orientation
 		boolean bPortrait = newConfig.orientation == Configuration.ORIENTATION_PORTRAIT;
-		GetOverrideGameActivity().nativeOnOrientationChanged(newConfig.orientation);
+		GetOverrideGameActivityWhenInitComplete().nativeOnOrientationChanged(newConfig.orientation);
 	}
 
 	public void onActivityResult(int requestCode, int resultCode, Intent data)
@@ -1661,8 +1618,8 @@ public class Engine
 		Log.v(TAG, "logOrder:" + logOrder + " * Engine::onActivityResult() - bCalledStart was = " + bCalledStart + ", bResuming was = " + bResuming + ", onPauseCounter=" + onPauseCounter + ", renderSurfaceValid: " + renderSurfaceValid + ", renderSurface = " + renderSurface );
 
 		// only used by FOnlineSubsystemGooglePlay so we changed signature to pass the ContextWrapper instead of activity here because we no longer want to hold onto any activity references...
-		GetOverrideGameActivity().nativeOnActivityResult(GameActivity.Get(), requestCode, resultCode, data);
-//		onResume();
+		GetOverrideGameActivityWhenInitComplete().nativeOnActivityResult(GameActivity.Get(), requestCode, resultCode, data);
+		//		??onResume();
 	}
 
 	public boolean onTouchEvent(MotionEvent event)
@@ -1680,7 +1637,7 @@ public class Engine
 			int pointerId = event.getPointerId(index);
 			int x = (int)event.getX();
 			int y = (int)event.getY();
-			switch (event.getAction()) {
+			switch (event.getActionMasked()) {
 				case MotionEvent.ACTION_DOWN:
 				case MotionEvent.ACTION_POINTER_DOWN:
 					action = TOUCHTYPE_BEGIN;
diff --git a/Engine/Plugins/Runtime/AndroidSingleInstanceService/Source/AndroidSingleInstanceService/Android/src/main/java/com/epicgames/makeaar/EngineFactory.java b/Engine/Plugins/Runtime/AndroidSingleInstanceService/Source/AndroidSingleInstanceService/Android/src/main/java/com/epicgames/makeaar/EngineFactory.java
--- a/Engine/Plugins/Runtime/AndroidSingleInstanceService/Source/AndroidSingleInstanceService/Android/src/main/java/com/epicgames/makeaar/EngineFactory.java
+++ b/Engine/Plugins/Runtime/AndroidSingleInstanceService/Source/AndroidSingleInstanceService/Android/src/main/java/com/epicgames/makeaar/EngineFactory.java
@@ -55,7 +55,7 @@ public class EngineFactory extends HandlerThread {
 		{
 			Log.d("UE", "EngineFactory::Reset() - releasing current engine resources, engine = " + Singleton.INSTANCE.engine);
 			//engine.onDestroy(543210);
-			Singleton.INSTANCE.engine.onDestroy(543210);
+			Singleton.INSTANCE.engine.onDestroy(543210, "EngineFactory::Reset()");
 		}
 		
 //		engineActivity = null;
diff --git a/Engine/Plugins/Runtime/AndroidSingleInstanceService/Source/AndroidSingleInstanceService/Android/src/main/java/com/epicgames/makeaar/GameActivityForMakeAAR.java b/Engine/Plugins/Runtime/AndroidSingleInstanceService/Source/AndroidSingleInstanceService/Android/src/main/java/com/epicgames/makeaar/GameActivityForMakeAAR.java
--- a/Engine/Plugins/Runtime/AndroidSingleInstanceService/Source/AndroidSingleInstanceService/Android/src/main/java/com/epicgames/makeaar/GameActivityForMakeAAR.java
+++ b/Engine/Plugins/Runtime/AndroidSingleInstanceService/Source/AndroidSingleInstanceService/Android/src/main/java/com/epicgames/makeaar/GameActivityForMakeAAR.java
@@ -83,19 +83,20 @@ public class GameActivityForMakeAAR extends com.epicgames.unreal.GameActivity
 		nativeSetCommandline(inCommandline);
 	}
 
-	public static Boolean isValidGameActivity()
-	{
-		return GameActivity.Get() != null;
-	}
-
 	public static GameActivityForMakeAAR Get()
 	{
 		return (GameActivityForMakeAAR)_gameActivityInstance;
 	}
 
-	@Override
+	// already present in super call and activity!
+	//public final void runOnUiThread(Runnable action)
+	//{
+	//	assert(activityContext != null);
+	//	activityContext.runOnUiThread(action);
+	//}
 
 
+	@Override
 	protected void RestartApplication(String RestartExtra)
 	{
 		super.RestartApplication(RestartExtra);
@@ -558,4 +559,10 @@ public class GameActivityForMakeAAR extends com.epicgames.unreal.GameActivity
 		gameActivitySetupInfo.GetVulkanInfoFromPackage(this);
 
 	}
+
+	@Override
+	public void onPause()
+	{
+		super.onPause();
+	}
 }
\ No newline at end of file
diff --git a/Engine/Plugins/Runtime/AndroidSingleInstanceService/Source/AndroidSingleInstanceService/Android/src/main/java/com/epicgames/makeaar/UnrealSharedInstanceService.java b/Engine/Plugins/Runtime/AndroidSingleInstanceService/Source/AndroidSingleInstanceService/Android/src/main/java/com/epicgames/makeaar/UnrealSharedInstanceService.java
--- a/Engine/Plugins/Runtime/AndroidSingleInstanceService/Source/AndroidSingleInstanceService/Android/src/main/java/com/epicgames/makeaar/UnrealSharedInstanceService.java
+++ b/Engine/Plugins/Runtime/AndroidSingleInstanceService/Source/AndroidSingleInstanceService/Android/src/main/java/com/epicgames/makeaar/UnrealSharedInstanceService.java
@@ -317,12 +317,12 @@ public class UnrealSharedInstanceService extends Service {
 
 					//if (bDoPause)
 					{
-						engineInstance.onPause(engineInstance.getCurrentContextID(), "UnrealMessageType.DetachExternalSurface, taskId=" + taskId);
+						engineInstance.QueuePause("UnrealMessageType.DetachExternalSurface, taskId=" + taskId);
 					}
 
 					//if (bDoStop)
 					{
-						engineInstance.onStop(engineInstance.getCurrentContextID(), "UnrealMessageType.DetachExternalSurface, taskId=" + taskId);
+						engineInstance.QueueStop("UnrealMessageType.DetachExternalSurface, taskId=" + taskId);
 					}
 
 					//engineInstance.sendConsoleCommand("t.maxfps 0.001");
diff --git a/Engine/Plugins/Runtime/AndroidSingleInstanceService/Source/AndroidSingleInstanceService/AndroidSingleInstanceService_UPL.xml b/Engine/Plugins/Runtime/AndroidSingleInstanceService/Source/AndroidSingleInstanceService/AndroidSingleInstanceService_UPL.xml
--- a/Engine/Plugins/Runtime/AndroidSingleInstanceService/Source/AndroidSingleInstanceService/AndroidSingleInstanceService_UPL.xml
+++ b/Engine/Plugins/Runtime/AndroidSingleInstanceService/Source/AndroidSingleInstanceService/AndroidSingleInstanceService_UPL.xml
@@ -24,7 +24,26 @@
 		</if>
 
 
+		<!-- disable embedding ASIS if shipping but "embed in shipping" not checked  -->
+		<setBoolIsEqual result="bIsShipping" arg1="$S(Configuration)" arg2="Shipping"/>
+		<if condition="bIsShipping">
+			<true>
+				<if condition="bIncludeInShipping">
+					<false>
+						<log text="Disabled embedding ASIS because shipping"/>
+						<setBool result="bEmbedASIS" value="false"/>
 
+						<!-- but we need to compile ASISProject to install data if using ASIS -->
+						<setBool result="bCompileASISProject" value="$B(bCreateASISProject)"/>
+					</false>
+				</if>
+				<if condition="bAllowExternalStartInShipping">
+					<false>
+						<setBool result="bExternalStartASIS" value="false"/>
+					</false>
+				</if>
+			</true>
+		</if>
 	
 	</init>
 
@@ -66,8 +85,13 @@
 						<if condition="bExternalStartASIS">
 							<true>
 								<log text="AndroidSingleInstanceService Manifest bExternalStartASIS so add intent filters and service..."/>
+								<!-- runtime permissions for foreground service -->
+								<addPermission android:name="android.permission.FOREGROUND_SERVICE"/>
+								<addPermission android:name="android.permission.FOREGROUND_SERVICE_DATA_SYNC"/>
+								<addPermission android:name="android.permission.POST_NOTIFICATIONS"/>
+								
 								<addElements tag="application">
-									<service android:name="com.epicgames.makeaar.UnrealSharedInstanceService" android:exported="true" android:process=":unrealSharedInstanceService" >
+									<service android:name="com.epicgames.makeaar.UnrealSharedInstanceService" android:exported="true" android:foregroundServiceType="dataSync" android:process=":unrealSharedInstanceService" >
 									</service>
 								</addElements>
 							</true>
@@ -113,7 +137,6 @@
 			<true>
 				<log text="AndroidSingleInstanceService resourceCopies copy service helper files. for BuildDirGradle= $S(BuildDir)/gradle/"/>
 				<copyDir src="$S(PluginDir)/Android/buildsystem" dst="$S(BuildDir)/gradle/buildsystem" />
-				<copyFile src="$S(PluginDir)/Android/src/main/java/com/epicgames/unreal/UEInterfaceHelpers.java" dst="$S(BuildDir)/src/com/epicgames/unreal/UEInterfaceHelpers.java" />
 			</true>
 		</if>
 			
@@ -122,13 +145,10 @@
 				<log text="AndroidSingleInstanceService resourceCopies copy service class files..."/>
 				<deleteFiles filespec="$S(BuildDir)/src/com/epicgames/makeaar" />
 				<copyDir src="$S(PluginDir)/Android/src/main/java/com/epicgames/makeaar" dst="$S(BuildDir)/src/com/epicgames/makeaar" />
-				<copyFile src="$S(PluginDir)/Android/src/main/java/com/epicgames/makeaar/UnrealSharedInstanceService.java" dst="$S(BuildDir)/src/com/epicgames/makeaar/UnrealSharedInstanceService.java" />
-
 			</true>
 			<false>
 				<log text="AndroidSingleInstanceService resourceCopies bEmbedASIS is false so remove files..."/>
 				<deleteFiles filespec="$S(BuildDir)/src/com/epicgames/makeaar" />
-				<deleteFiles filespec="$S(BuildDir)/src/com/epicgames/makeaar/UnrealSharedInstanceService*.java" />
 			</false>
 		</if>
 	</resourceCopies>
@@ -138,6 +158,9 @@
 			<true>
 				<log text="AndroidSingleInstanceService create ASISStub project..."/>
 				<log text="AndroidSingleInstanceService bCreateASISProject, AndroidPackageName = $S(AndroidPackageName)"/>
+				<copyDir src="$S(PluginDir)/../ASISStub" dst="$S(BuildDir)/gradle/ASISProject" />
+				<copyDir src="$S(BuildDir)/gradle/app/src/main/res" dst="$S(BuildDir)/gradle/ASISProject/app/src/main/res" />
+				<copyDir src="$S(BuildDir)/gradle/gradle" dst="$S(BuildDir)/gradle/ASISProject/gradle" />
 			</true>
 		</if>
 	</gradleCopies>
@@ -146,27 +169,69 @@
 		<if condition="bCompileASISProject">
 			<true>
 				<log text="AndroidSingleInstanceService buildGradleAdditions bCompileASISProject..."/>
+				<insert>
+<![CDATA[
+// --- Begin ASISProject ---
+tasks.whenTaskAdded { task ->
+	if (task.name == 'assembleRelease') {
+		tasks.create(name:'ueASISProjectAssembleRelease', type: Exec) {
+			description "Compile ASISProject for Release"
+			workingDir "$rootDir" + "/ASISProject"
+			if (System.getProperty('os.name').toLowerCase(Locale.ROOT).contains('windows')) {
+				commandLine 'cmd', '/c', 'gradlew.bat', 'app:assembleRelease'
+			} else {
+				commandLine 'sh', '-c', 'chmod +x ./gradlew && ./gradlew app:assembleRelease'
+			}
+		}
+		task.finalizedBy 'ueASISProjectAssembleRelease'
+	}
+	if (task.name == 'assembleDebug') {
+		tasks.create(name:'ueASISProjectAssembleDebug', type: Exec) {
+			description "Compile ASISProject for Debug"
+			workingDir "$rootDir" + "/ASISProject"
+			if (System.getProperty('os.name').toLowerCase(Locale.ROOT).contains('windows')) {
+				commandLine 'cmd', '/c', 'gradlew.bat', 'app:assembleDebug'
+			} else {
+				commandLine 'sh', '-c', 'chmod +x ./gradlew && ./gradlew app:assembleDebug'
+			}
+		}
+		task.finalizedBy 'ueASISProjectAssembleDebug'
+	}
+	if (task.name == 'bundleRelease') {
+		tasks.create(name:'ueASISProjectBundleRelease', type: Exec) {
+			description "Compile ASISProject for Release"
+			workingDir "$rootDir" + "/ASISProject"
+			if (System.getProperty('os.name').toLowerCase(Locale.ROOT).contains('windows')) {
+				commandLine 'cmd', '/c', 'gradlew.bat', 'app:assembleRelease'
+			} else {
+				commandLine 'sh', '-c', 'chmod +x ./gradlew && ./gradlew app:assembleRelease'
+			}
+		}
+		task.finalizedBy 'ueASISProjectBundleRelease'
+	}
+	if (task.name == 'bundleDebug') {
+		tasks.create(name:'ueASISProjectBundleDebug', type: Exec) {
+			description "Compile ASISProject for Debug"
+			workingDir "$rootDir" + "/ASISProject"
+			if (System.getProperty('os.name').toLowerCase(Locale.ROOT).contains('windows')) {
+				commandLine 'cmd', '/c', 'gradlew.bat', 'app:assembleDebug'
+			} else {
+				commandLine 'sh', '-c', 'chmod +x ./gradlew && ./gradlew app:assembleDebug'
+			}
+		}
+		task.finalizedBy 'ueASISProjectBundleDebug'
+	}
+}
+// --- End ASISProject ---
+]]>
+				</insert>
 			</true>
 		</if>
 
 		<if condition="bEmbedASIS">
 			<true>
 				<log text="AndroidSingleInstanceService buildGradleAdditions bEmbedASIS. for BuildDir= $S(BuildDir)"/>
-				<insert>
-					<![CDATA[
-
-apply from: '../buildsystem/config.gradle'
-apply from: '../buildsystem/dependencies.gradle'
 
-		android {
-		defaultConfig {
-		buildConfigField "String", "OBB_MODULE_NAME", config.obbModuleName
-		buildConfigField "String", "OBB_FILE_NAME", config.obbFileLocation
-		buildConfigField "String", "AAR_VERSION", config.AAR_VERSION
-		}
-		}
-]]>
-				</insert>
 			</true>
 		</if>
 	</buildGradleAdditions>
@@ -174,6 +239,108 @@ apply from: '../buildsystem/dependencies.gradle'
 	<settingsGradleAdditions>
 		<if condition="bCreateASISProject">
 			<true>
+				<insert>
+	<![CDATA[
+	// copy files into ASISStub
+	File asis_infile = new File(rootProject.projectDir, 'gradle.properties')
+	File asis_outfile = new File(rootProject.projectDir, 'ASISProject/gradle.properties')
+	def asis_line = "";
+	asis_outfile.withWriter { writer ->
+		asis_infile.withReader { reader ->
+			while ((asis_line = reader.readLine()) != null) {
+				if (asis_line.startsWith('OUTPUT_FILENAME=')) {
+					writer.writeLine("OUTPUT_FILENAME=ASIS_" + asis_line.substring(16));
+					continue;
+				}
+				else if (asis_line.startsWith('OUTPUT_BUNDLEFILENAME=')) {
+					writer.writeLine("OUTPUT_BUNDLEFILENAME=ASIS_" + asis_line.substring(22));
+					continue;
+				}
+				writer.writeLine("${asis_line}")
+			}
+		}
+	}
+
+	asis_infile = new File(rootProject.projectDir, 'app/buildAdditions.gradle')
+	asis_outfile = new File(rootProject.projectDir, 'ASISProject/app/buildAdditions.gradle')
+	asis_line = "";
+	def bRemovingASISBlock = false;
+	asis_outfile.withWriter { writer ->
+		asis_infile.withReader { reader ->
+			while ((asis_line = reader.readLine()) != null) {
+				if (bRemovingASISBlock)
+				{
+					if (asis_line.startsWith('// --- End ASISProject ---'))
+					{
+						bRemovingASISBlock = false;
+					}
+					continue;
+				}
+				if (asis_line.startsWith('// --- Begin ASISProject ---'))
+				{
+					bRemovingASISBlock = true;
+					continue;
+				}
+				writer.writeLine("${asis_line}")
+			}
+		}
+	}
+
+	asis_infile = new File(rootProject.projectDir, 'app/src/main/AndroidManifest.xml')
+	asis_outfile = new File(rootProject.projectDir, 'ASISProject/app/src/main/AndroidManifest.xml')
+	def asis_hasForeServiceDataPerm = false;
+	def asis_hasForeServicePerm = false;
+	def asis_hasNotificationsPerm = false;
+	def asis_hasService = false
+	def asis_hasActivity = false
+	asis_outfile.withWriter { writer ->
+		asis_infile.withReader { reader ->
+			while ((asis_line = reader.readLine()) != null) {
+				if (asis_line.contains('android.permission.FOREGROUND_SERVICE_DATA_SYNC')) {
+					asis_hasForeServiceDataPerm = true
+				} else if (asis_line.contains('android.permission.FOREGROUND_SERVICE')) {
+					asis_hasForeServicePerm = true
+				} else if (asis_line.contains('android.permission.POST_NOTIFICATIONS')) {
+					asis_hasNotificationsPerm = true
+				} else if (asis_line.contains('com.epicgames.makeaar.UnrealSharedInstanceService')) {
+					asis_hasService = true
+				} else if (asis_line.contains('com.epicgames.makeaar.ASISExampleGameActivity')) {
+					asis_hasActivity = true
+				} else if (asis_line.contains('</application>')) {
+
+					if (asis_hasActivity == false) {
+						writer.writeLine('    <activity android:name="com.epicgames.makeaar.GameActivity" android:exported="true" android:screenOrientation="sensor">')
+						writer.writeLine('    	<intent-filter>')
+						writer.writeLine('    		<action android:name="android.intent.action.MAIN" />')
+						writer.writeLine('    		<category android:name="android.intent.category.LAUNCHER" />')
+						writer.writeLine('    	</intent-filter>')
+						writer.writeLine('    </activity>')
+					}
+				}
+				writer.writeLine("${asis_line}")
+			}
+		}
+	}
+
+	def asis_files = ["app/aar-imports.gradle", "app/abi.gradle", "app/build.gradle",
+		"baseBuildAdditions.gradle", "build.gradle", "buildscriptAdditions.gradle", "local.properties",
+		"gradlew", "gradlew.bat", "rungradle.bat"] as String[]
+	asis_files.eachWithIndex { v, i ->
+		asis_infile = new File(rootProject.projectDir, v)
+		if (asis_infile.exists()) {
+			asis_outfile = new File(rootProject.projectDir, "ASISProject/" + v)
+			asis_outfile.withWriter { writer ->
+				asis_infile.withReader { reader ->
+					while ((asis_line = reader.readLine()) != null) {
+						writer.writeLine("${asis_line}")
+					}
+				}
+			}
+		}
+	}
+
+]]>
+				</insert>
 			</true>
 		</if>
 	</settingsGradleAdditions>
@@ -181,6 +348,10 @@ apply from: '../buildsystem/dependencies.gradle'
 	<gameActivityImportAdditions>
 		<if condition="bEmbedASIS">
 			<true>
+				<insert>
+import java.util.Set;
+
+				</insert>
 			</true>
 		</if>
 	</gameActivityImportAdditions>
@@ -189,12 +360,181 @@ apply from: '../buildsystem/dependencies.gradle'
 		<if condition="bEmbedASIS">
 			<true>
 				<insert>
-				</insert>
-				<insertNewline/>
-				<insert>
+	<![CDATA[
+	public int AndroidThunkJava_AndroidSingleInstanceServer_IsRunning()
+	{
+        Set<Thread> threads = Thread.getAllStackTraces().keySet();
+		int ConnectionType = 0;
+        for (Thread thread : threads) {
+            String name = thread.getName();
+            if (name.equals("UnrealASIS-USB"))
+            {
+                if (thread.isAlive()) {
+					ConnectionType |= 1;
+                }
+            }
+            else if (name.equals("UnrealASIS-WiFi"))
+            {
+                if (thread.isAlive()) {
+					ConnectionType |= 2;
+                }
+            }
+        }
+		
+		return ConnectionType;
+	}
+
+	public boolean AndroidThunkJava_AndroidSingleInstanceServer_Stop(boolean bUSB, boolean bNetwork)
+	{
+        Set<Thread> threads = Thread.getAllStackTraces().keySet();
+        boolean bFoundThread = false;
+        for (Thread thread : threads) {
+            String name = thread.getName();
+            if (bUSB && name.equals("UnrealASIS-USB"))
+            {
+                if (thread.isAlive()) {
+					thread.interrupt();
+                    bFoundThread = true;
+                }
+            }
+            else if (bNetwork && name.equals("UnrealASIS-WiFi"))
+            {
+                if (thread.isAlive()) {
+					thread.interrupt();
+                    bFoundThread = true;
+                }
+            }
+        }
+		
+		return bFoundThread;
+	}
+
+	public boolean AndroidThunkJava_AndroidSingleInstanceServer_Start(boolean bUSB, boolean bNetwork, int port)
+	{
+		boolean bUSBRunning = false;
+		boolean bWiFiRunning = false;
+		
+		if (port == 0)
+		{
+			port = 57099;
+		}
+		
+        Set<Thread> threads = Thread.getAllStackTraces().keySet();
+        boolean bFoundThread = false;
+        for (Thread thread : threads) {
+            String name = thread.getName();
+            if (name.equals("UnrealASIS-USB"))
+            {
+                if (thread.isAlive()) {
+                    bUSBRunning = true;
+                }
+            }
+            else if (name.equals("UnrealASIS-WiFi"))
+            {
+                if (thread.isAlive()) {
+                    bWiFiRunning = true;
+                    break;
+                }
+            }
+        }
+
+		if (bUSBRunning && bWiFiRunning)
+		{
+			return true;
+		}
+
+		String packageName = getPackageName();
+		PackageManager pm = getPackageManager();
+		
+		int versionCode = 0;
+		try
+		{
+			versionCode = pm.getPackageInfo(packageName, 0).versionCode;
+		}
+		catch (Exception e)
+		{
+			// if the above failed, then, we can't use obbs
+		}
+
+		String EngineVersion = "5.0.0";
+		String ProjectName = packageName;
+		ProjectName = ProjectName.substring(ProjectName.lastIndexOf('.') + 1);
+		boolean IsShipping = false;
+		boolean UseExternalFilesDir = false;
+		boolean PublicLogFiles = false;
+		boolean AllowNetworkConnection = false;
+		try {
+			ApplicationInfo ai = pm.getApplicationInfo(packageName, PackageManager.GET_META_DATA);
+			Bundle bundle = ai.metaData;
+
+			if (bundle.containsKey("com.epicgames.unreal.GameActivity.EngineVersion"))
+			{
+				EngineVersion = bundle.getString("com.epicgames.unreal.GameActivity.EngineVersion");
+			}
+			if (bundle.containsKey("com.epicgames.unreal.GameActivity.ProjectName"))
+			{
+				ProjectName = bundle.getString("com.epicgames.unreal.GameActivity.ProjectName");
+			}
+			if (bundle.containsKey("com.epicgames.unreal.GameActivity.BuildConfiguration"))
+			{
+				String Configuration = bundle.getString("com.epicgames.unreal.GameActivity.BuildConfiguration");
+				IsShipping = Configuration.contains("Shipping");
+			}
+			if (bundle.containsKey("com.epicgames.unreal.GameActivity.bUseExternalFilesDir"))
+			{
+				UseExternalFilesDir = bundle.getBoolean("com.epicgames.unreal.GameActivity.bUseExternalFilesDir");
+			}
+			if (bundle.containsKey("com.epicgames.unreal.GameActivity.bPublicLogFiles"))
+			{
+				PublicLogFiles = bundle.getBoolean("com.epicgames.unreal.GameActivity.bPublicLogFiles");
+			}
+			if (bundle.containsKey("com.epicgames.unreal.RemoteFileManager.bAllowNetworkConnection"))
+			{
+				AllowNetworkConnection = bundle.getBoolean("com.epicgames.unreal.RemoteFileManager.bAllowNetworkConnection");
+			}
+		}
+		catch (PackageManager.NameNotFoundException | NullPointerException e)
+		{
+			Log.debug("Error when accessing application metadata");
+		}
+
+		String internal = getFilesDir().getAbsolutePath();
+		String external = getExternalFilesDir(null).getAbsolutePath();
+		String storage = android.os.Environment.getExternalStorageDirectory().getAbsolutePath();
+		String obbdir = getObbDir().getAbsolutePath();
+
+		if (bUSB && !bUSBRunning)
+		{
+			//RemoteFileManager fileManager = new RemoteFileManager(true, port, internal, external, storage, obbdir, packageName, versionCode, ProjectName, EngineVersion, IsShipping, PublicLogFiles);
+			//new Thread(fileManager, "UnrealASIS-USB").start();
+		}
+		if (AllowNetworkConnection && bNetwork && !bWiFiRunning)
+		{
+			//RemoteFileManager fileManager = new RemoteFileManager(false, port, internal, external, storage, obbdir, packageName, versionCode, ProjectName, EngineVersion, IsShipping, PublicLogFiles);
+			//new Thread(fileManager, "UnrealASIS-WiFi").start();
+		}
+		
+		return true;
+	}
+]]>
 				</insert>
 			</true>
 			<false>
+	public boolean AndroidThunkJava_AndroidSingleInstanceServer_IsRunning()
+	{
+	return false;
+	}
+
+	public boolean AndroidThunkJava_AndroidSingleInstanceServer_Start()
+	{
+	return false;
+	}
+
+	public boolean AndroidThunkJava_AndroidSingleInstanceServer_Stop()
+	{
+	return false;
+	}
+
 			</false>
 		</if>
 	</gameActivityClassAdditions>
diff --git a/Engine/Plugins/Runtime/AndroidSingleInstanceService/Source/AndroidSingleInstanceService/Private/AndroidSingleInstanceServiceBPLibrary.cpp b/Engine/Plugins/Runtime/AndroidSingleInstanceService/Source/AndroidSingleInstanceService/Private/AndroidSingleInstanceServiceBPLibrary.cpp
--- a/Engine/Plugins/Runtime/AndroidSingleInstanceService/Source/AndroidSingleInstanceService/Private/AndroidSingleInstanceServiceBPLibrary.cpp
+++ b/Engine/Plugins/Runtime/AndroidSingleInstanceService/Source/AndroidSingleInstanceService/Private/AndroidSingleInstanceServiceBPLibrary.cpp
@@ -278,6 +278,8 @@ bool UAndroidSingleInstanceServiceBPLibrary::EnableButton(const FString& Key, bo
 #define EVENTTYPE_PRE_LOAD_MAP				5
 #define EVENTTYPE_POST_LOAD_MAP				6
 #define EVENTTYPE_ACTION					7
+#define EVENTTYPE_ACTIVITYCHANGED			8
+#define EVENTTYPE_APP_SUSPENDED				9
 
 void UAndroidSingleInstanceServiceBPLibrary::SendAndroidEvent(int32 EventId, const FString& Param1, int32 Param2, int32 Param3, float Param4)
 {
diff --git a/Engine/Source/Runtime/ApplicationCore/Private/Android/AndroidWindow.cpp b/Engine/Source/Runtime/ApplicationCore/Private/Android/AndroidWindow.cpp
--- a/Engine/Source/Runtime/ApplicationCore/Private/Android/AndroidWindow.cpp
+++ b/Engine/Source/Runtime/ApplicationCore/Private/Android/AndroidWindow.cpp
@@ -334,6 +334,14 @@ void FAndroidWindow::ReleaseWindowRef(ANativeWindow* InWindow)
 #if USE_ANDROID_EVENTS
 	check(IsInAndroidEventThread());
 #endif
+
+#if USE_ANDROID_STANDALONE
+	if (GAndroidWindowOverride && InWindow != GAndroidWindowOverride)
+	{
+		STANDALONE_DEBUG_LOGf(LogAndroid, TEXT("SetHardwareWindow_EventThread -> InWindow(%p) is not current GAndroidWindowOverride(%p)"), InWindow, GAndroidWindowOverride);
+	}
+#endif
+
 	//using raw native window handle for now. Could be changed to use AndroidWindow later if needed
 	NativeWindow = InWindow; 
 }
diff --git a/Engine/Source/Runtime/Launch/Private/Android/LaunchAndroid.cpp b/Engine/Source/Runtime/Launch/Private/Android/LaunchAndroid.cpp
--- a/Engine/Source/Runtime/Launch/Private/Android/LaunchAndroid.cpp
+++ b/Engine/Source/Runtime/Launch/Private/Android/LaunchAndroid.cpp
@@ -182,6 +182,14 @@ extern FString GFilePathBase;
 /** The global EngineLoop instance */
 FEngineLoop	GEngineLoop;
 
+#if USE_ANDROID_STANDALONE
+	FString GAndroidCommandLine;
+#endif
+
+extern void* GAndroidWindowOverride;
+extern int GSurfaceViewWidth;
+extern int GSurfaceViewHeight;
+
 static bool bDidCompleteEngineInit = false;
 
 bool GShowConsoleWindowNextTick = false;
@@ -216,13 +224,6 @@ static TAutoConsoleVariable<int32> CVarEnableSustainedPerformanceMode(
 	TEXT("  1: Enabled"),
 	ECVF_Default);
 
-static TAutoConsoleVariable<int32> CVarUseGameThreadForCommands(
-	TEXT("Android.UseGameThread"),
-	0,
-	TEXT("Use GameThread for incoming commands\n")
-	TEXT("  0: Disabled (default)\n")
-	TEXT("  1: Enabled"),
-	ECVF_Default);
 
 extern void AndroidThunkCpp_SetSustainedPerformanceMode(bool);
 static void SetSustainedPerformanceMode()
@@ -244,6 +245,7 @@ static FEvent* EventHandlerEvent = NULL;
 // Wait for Java onCreate to complete before resume main init
 static volatile bool GResumeMainInit = false;
 volatile bool GEventHandlerInitialized = false;
+
 static void IssueConsoleCommand(FString Command)
 {
 	static TArray<FString> GPendingConsoleCommands;
@@ -288,7 +290,6 @@ static bool bIgnorePauseOnDownloaderStart = false;
 bool bReadyToProcessEvents = false;
 static int GAndroidWindowLockRefCount = 0;
 static const bool bForce_GAndroidWindowLock = true;
-extern void* GAndroidWindowOverride;
 
 void GAndroidWindowLock_Lock(FString calledBy)
 {
@@ -375,8 +376,6 @@ void FPlatformMisc::UnlockAndroidWindow()
 	GAndroidWindowLock_Unlock("UnlockAndroidWindow");
 	
 	bReadyToProcessEvents = true;
-	//bAppIsActive_EventThread = true;
-	
 }
 
 JNI_METHOD void Java_com_epicgames_unreal_GameActivity_nativeResumeMainInit(JNIEnv* jenv, jobject thiz)
@@ -415,9 +414,7 @@ void InitHMDs()
 }
 
 extern AAssetManager* AndroidThunkCpp_GetAssetManager();
-#if USE_ANDROID_STANDALONE
-	FString GAndroidCommandLine;
-#endif
+
 /**
  * Shuts down the engine
  */
@@ -584,9 +581,22 @@ static void ApplyAndroidCompatConfigRules()
 	}
 }
 
+#if USE_ANDROID_STANDALONE
+static void TickDeferredAppCommands(FString calledBy);
+static void IssueAppCommand(long Command);
+#endif
+
+
 //Main function called from the android entry point
+#if USE_ANDROID_STANDALONE
+void* AndroidMain(void* param)
+{
+	struct android_app* state = (struct android_app*)param;
+#else
 int32 AndroidMain(struct android_app* state)
 {
+#endif
+
 	BootTimingPoint("AndroidMain");
 
 	FPlatformMisc::LowLevelOutputDebugString(TEXT("Entered AndroidMain()\n"));
@@ -744,6 +754,9 @@ int32 AndroidMain(struct android_app* state)
 	FPlatformMisc::LowLevelOutputDebugString(TEXT("Created sync event\n"));
 	FAppEventManager::GetInstance()->SetEventHandlerEvent(EventHandlerEvent);
 
+	STANDALONE_DEBUG_LOGf(LogAndroid, TEXT("LogLoad::AndroidMain (after SetEventHandlerEvent) IsDebuggerPresent = %d, GAlwaysReportCrash = %d, bEnableGuarded = %d, GIsGuarded = %d"), FPlatformMisc::IsDebuggerPresent(), GAlwaysReportCrash, bEnableGuarded, GIsGuarded);
+
+
 	// ready for onCreate to complete
 	GEventHandlerInitialized = true;
 
@@ -753,6 +766,7 @@ int32 AndroidMain(struct android_app* state)
 	// Have to use a special initialize if using the PersistentStorageManager
 	IPlatformFile::GetPlatformPhysical().Initialize(nullptr, FCommandLine::Get());
 
+
 	{
 		SCOPED_BOOT_TIMING("Wait for GAndroidWindowLock.Lock()");
 		// wait for a valid window
@@ -762,8 +776,11 @@ int32 AndroidMain(struct android_app* state)
 
 		UE_LOG(LogAndroid, Log, TEXT("PreInit android HW window lock. bAppIsActive_EventThread=%d"), bAppIsActive_EventThread);
 		DEVELOPER_LOG_COMMANDCB_CASE(PreInitWindowLock);
+#if !USE_ANDROID_STANDALONE
 		GAndroidWindowLock_Lock("AndroidMain lock");
+#endif
 	}
+
 	FPlatformMisc::LowLevelOutputDebugString( TEXT("After GAndroidWindowLock in AndroidMain"));
 
 	FDelegateHandle ConfigReadyHandle = FCoreDelegates::TSConfigReadyForUse().AddStatic(&ApplyAndroidCompatConfigRules);
@@ -776,11 +793,14 @@ int32 AndroidMain(struct android_app* state)
 	if (PreInitResult != 0)
 	{
 		checkf(false, TEXT("Engine Preinit Failed"));
+#if USE_ANDROID_STANDALONE
+		return nullptr;
+#else
 		return PreInitResult;
+#endif
 	}
 
 #if !USE_ANDROID_STANDALONE
-
 	// register callback for native window resize
 	if (GAndroidEnableNativeResizeEvent)
 	{
@@ -801,7 +821,7 @@ int32 AndroidMain(struct android_app* state)
 
 	FAppEventManager::GetInstance()->SetEmptyQueueHandlerEvent(FPlatformProcess::GetSynchEventFromPool(false));
 
-	GEngineLoop.Init();
+	ErrorLevel = GEngineLoop.Init();
 
 	bDidCompleteEngineInit = true;
 
@@ -840,11 +860,7 @@ int32 AndroidMain(struct android_app* state)
 		FAndroidStats::UpdateAndroidStats();
 
 		FAppEventManager::GetInstance()->Tick();
-		if (!FAppEventManager::GetInstance()->IsGamePaused()
-#if USE_ANDROID_STANDALONE
-		//	&& FAppEventManager::GetInstance()->IsGameInFocus()
-#endif
-			)
+		if (!FAppEventManager::GetInstance()->IsGamePaused())
 		{
 			GEngineLoop.Tick();
 		}
@@ -877,13 +893,12 @@ int32 AndroidMain(struct android_app* state)
 
 	FPlatformMisc::RequestExit(true, TEXT("AndroidMain"));
 #endif
-	return 0;
-}
-void* AndroidMain(void* param)
-{
-	struct android_app* state = (struct android_app*)param;
-	AndroidMain(state);
+
+#if USE_ANDROID_STANDALONE
 	return nullptr;
+#else
+	return ErrorLevel;
+#endif
 }
 
 struct AChoreographer;
@@ -1001,8 +1016,13 @@ static uint32 EventThreadID = 0;
 
 bool IsInAndroidEventThread()
 {
+#if USE_ANDROID_STANDALONE
+	//@TODO: for now always return true to avoid check failures
+	return true;
+#else
 	check(EventThreadID != 0);
 	return EventThreadID == FPlatformTLS::GetCurrentThreadId();
+#endif
 }
 
 static void* AndroidEventThreadWorker( void* param )
@@ -1033,17 +1053,11 @@ static void* AndroidEventThreadWorker( void* param )
 
 	TheChoreographer.SetupChoreographer();
 
-	
+	GEventHandlerInitialized = true;
 
 	// window is initially invalid/locked.
-	extern void* GAndroidWindowOverride;
-	UE_LOG(LogAndroid, Log, TEXT("AndroidEventThreadWorker, Initial HW window lock. GAndroidWindowOverride=%p"), GAndroidWindowOverride);
-#if USE_ANDROID_STANDALONE
-	if (GAndroidWindowOverride != nullptr)
-#endif
-	{
-		GAndroidWindowLock_Lock("AndroidEventThreadWorker");
-	}
+	UE_LOG(LogAndroid, Log, TEXT("AndroidEventThreadWorker, Initial HW window lock."));
+	//GAndroidWindowLock_Lock("AndroidEventThreadWorker");
 
 	DEVELOPER_LOG_COMMANDCB_CASE(AndroidEventThreadWorker_BeforeWhile);
 
@@ -1051,7 +1065,7 @@ static void* AndroidEventThreadWorker( void* param )
 	//continue to process events until the engine is shutting down
 	while (!IsEngineExitRequested())
 	{
-//		FPlatformMisc::LowLevelOutputDebugString(TEXT("AndroidEventThreadWorker"));
+		//FPlatformMisc::LowLevelOutputDebugString(TEXT("AndroidEventThreadWorker"));
 
 		AndroidProcessEvents(state);
 
@@ -1059,10 +1073,7 @@ static void* AndroidEventThreadWorker( void* param )
 	}
 	DEVELOPER_LOG_COMMANDCB_CASE(AndroidEventThreadWorker_AfterWhile);
 
-	GAndroidWindowLock_Unlock("AndroidEventThreadWorker");
-#if USE_ANDROID_STANDALONE
-	bReadyToProcessEvents = true;
-#endif
+	//GAndroidWindowLock_Unlock("AndroidEventThreadWorker");
 
 	UE_LOG(LogAndroid, Log, TEXT("AndroidEventThreadWorker->Exiting"));
 	STANDALONE_DEBUG_LOGf(LogAndroid, TEXT("Exiting AndroidEventThreadWorker"));
@@ -1530,28 +1541,19 @@ FAppEventData::FAppEventData(ANativeWindow* WindowIn)
 		UE_LOG(LogAndroid, Log, TEXT("FAppEventData was passed a NULL WindowsIn pointer."));
 		return;
 	}
+
 	//check(WindowIn);
-	extern void* GAndroidWindowOverride;
-	extern int GSurfaceViewWidth;
-	extern int GSurfaceViewHeight;
-	if (GAndroidWindowOverride != nullptr)
+	if (GAndroidWindowOverride != nullptr )
 	{
-		if (GSurfaceViewWidth>0 && GSurfaceViewHeight>0)
-		{
-			WindowWidth = GSurfaceViewWidth;	//ANativeWindow_getWidth((ANativeWindow*)GAndroidWindowOverride);
-			WindowHeight = GSurfaceViewHeight;	//ANativeWindow_getHeight((ANativeWindow*)GAndroidWindowOverride);
-		}
-        else
-        {
-			GSurfaceViewWidth = WindowWidth = ANativeWindow_getWidth((ANativeWindow*)GAndroidWindowOverride);
-			GSurfaceViewHeight = WindowHeight = ANativeWindow_getHeight((ANativeWindow*)GAndroidWindowOverride);
-		}
+		GSurfaceViewWidth = WindowWidth = ANativeWindow_getWidth((ANativeWindow*)GAndroidWindowOverride);
+		GSurfaceViewHeight = WindowHeight = ANativeWindow_getHeight((ANativeWindow*)GAndroidWindowOverride);
 	}
 	else
 	{
-	WindowWidth = ANativeWindow_getWidth(WindowIn);
-	WindowHeight = ANativeWindow_getHeight(WindowIn);
+		GSurfaceViewWidth = WindowWidth = ANativeWindow_getWidth(WindowIn);
+		GSurfaceViewHeight = WindowHeight = ANativeWindow_getHeight(WindowIn);
 	}
+
 	if (WindowWidth <= 0)
 	{
 		UE_LOG(LogAndroid, Log, TEXT("FAppEventData was passed a window with ZERO Height."));
@@ -1574,7 +1576,7 @@ static void ActivateApp_EventThread()
 	if (bAppIsActive_EventThread)
 	{
 		// Seems this can occur.
-		STANDALONE_DEBUG_LOGf(LogAndroid, TEXT("WARNING -- event thread, activate app, ALREADY have HW window lock according to bAppIsActive_EventThread. bReadyToProcessEvents=%d, bAppIsActive_EventThread=%d"), bReadyToProcessEvents, bAppIsActive_EventThread);
+		STANDALONE_DEBUG_LOGf(LogAndroid, TEXT("UNEXPECTED -- event thread, activate app, ALREADY have HW window lock according to bAppIsActive_EventThread. bReadyToProcessEvents=%d, bAppIsActive_EventThread=%d"), bReadyToProcessEvents, bAppIsActive_EventThread);
 #if !USE_ANDROID_STANDALONE
 		return;
 #endif
@@ -1583,9 +1585,7 @@ static void ActivateApp_EventThread()
 	// Unlock window when we're ready.
 	UE_LOG(LogAndroid, Log, TEXT("event thread, activate app, unlocking HW window"));
 	GAndroidWindowLock_Unlock("ActivateApp_EventThread");
-#if USE_ANDROID_STANDALONE
-	bReadyToProcessEvents = true;
-#endif
+
 	// wake the GT up.
 	FAppEventManager::GetInstance()->EnqueueAppEvent(APP_EVENT_STATE_APP_ACTIVATED);
 
@@ -1606,13 +1606,10 @@ static void ActivateApp_EventThread()
 	}
 
 	FThreadHeartBeat::Get().ResumeHeartBeat(true);
-#if !USE_ANDROID_STANDALONE
-
 	FPreLoadScreenManager::EnableRendering(true);
 
 	extern void AndroidThunkCpp_ShowHiddenAlertDialog();
 	AndroidThunkCpp_ShowHiddenAlertDialog();	
-#endif
 }
 
 extern void BlockRendering();
@@ -1622,28 +1619,26 @@ static void SuspendApp_EventThread()
 	DEVELOPER_LOG_COMMANDCB_CASE(SuspendApp_EventThread);
 	if (!bAppIsActive_EventThread)
 	{
-		STANDALONE_DEBUG_LOGf(LogAndroid, TEXT("event thread, SuspendApp_EventThread called but aborted, GAndroidWindowLockRefCount=%d, bReadyToProcessEvents=%d, bAppIsActive_EventThread=%d, IsInGameThread=%d"), GAndroidWindowLockRefCount, bReadyToProcessEvents, bAppIsActive_EventThread, IsInGameThread());
-#if !USE_ANDROID_STANDALONE
+		STANDALONE_DEBUG_LOGf(LogAndroid, TEXT("SuspendApp_EventThread -> UNEXPECTED: event thread, SuspendApp_EventThread called but aborted, GAndroidWindowLockRefCount=%d, bReadyToProcessEvents=%d, bAppIsActive_EventThread=%d, IsInGameThread=%d"), GAndroidWindowLockRefCount, bReadyToProcessEvents, bAppIsActive_EventThread, IsInGameThread());
 		return;
-#endif
 	}
 	bAppIsActive_EventThread = false;
 	// Lock the window, this prevents event thread from removing the window whilst the RHI initializes.
 
-	UE_LOG(LogAndroid, Log, TEXT("event thread, suspending app, acquiring HW window lock."));
+	UE_LOG(LogAndroid, Log, TEXT("SuspendApp_EventThread -> event thread, suspending app, acquiring HW window lock."));
 	GAndroidWindowLock_Lock("SuspendApp_EventThread");
-	STANDALONE_DEBUG_LOGf(LogAndroid, TEXT("event thread, suspending app, acquiring HW window lock. bReadyToProcessEvents=%d, bAppIsActive_EventThread=%d, IsInGameThread=%d"), bReadyToProcessEvents, bAppIsActive_EventThread, IsInGameThread());
+	STANDALONE_DEBUG_LOGf(LogAndroid, TEXT("SuspendApp_EventThread -> event thread, suspending app, acquiring HW window lock. bReadyToProcessEvents=%d, bAppIsActive_EventThread=%d, IsInGameThread=%d"), bReadyToProcessEvents, bAppIsActive_EventThread, IsInGameThread());
 
 	if (bReadyToProcessEvents == false)
 	{
 		// App has stopped before we can process events.
 		// AndroidLaunch will lock GAndroidWindowLock, and set bReadyToProcessEvents when we are able to block the RHI and queue up other events.
 		// we ignore events until this point as acquiring GAndroidWindowLock means requires the window to be properly initialized.
-		UE_LOG(LogAndroid, Log, TEXT("event thread, app not yet ready."));
+		UE_LOG(LogAndroid, Warning, TEXT("SuspendApp_EventThread -> UNEXPECTED: event thread, app not yet ready."));
 		return;
 	}
 
-	STANDALONE_DEBUG_LOGf(LogAndroid, TEXT("event thread, suspending app, ACQUIRED HW window lock. bReadyToProcessEvents=%d, bAppIsActive_EventThread=%d"), bReadyToProcessEvents, bAppIsActive_EventThread);
+	STANDALONE_DEBUG_LOGf(LogAndroid, TEXT("SuspendApp_EventThread -> event thread, suspending app, ACQUIRED HW window lock. bReadyToProcessEvents=%d, bAppIsActive_EventThread=%d"), bReadyToProcessEvents, bAppIsActive_EventThread);
 
 	TSharedPtr<FEvent, ESPMode::ThreadSafe> EMDoneTrigger = MakeShareable(FPlatformProcess::GetSynchEventFromPool(), [](FEvent* EventToDelete)
 	{
@@ -1654,7 +1649,7 @@ static void SuspendApp_EventThread()
 	// This ensures any tasks that require a window handle will have it before we block the RT on the invalid window.
 	FAppEventManager::GetInstance()->EnqueueAppEvent(APP_EVENT_RUN_CALLBACK, FAppEventData([EMDoneTrigger]()
 	{
-		UE_LOG(LogAndroid, Log, TEXT("performing app backgrounding callback. %p"), EMDoneTrigger.Get());
+		UE_LOG(LogAndroid, Log, TEXT("SuspendApp_EventThread -> performing app backgrounding callback. %p"), EMDoneTrigger.Get());
 
 		FCoreDelegates::ApplicationWillDeactivateDelegate.Broadcast();
 		FCoreDelegates::ApplicationWillEnterBackgroundDelegate.Broadcast();
@@ -1674,12 +1669,12 @@ static void SuspendApp_EventThread()
 	FThreadHeartBeat::Get().SuspendHeartBeat(true);
 
 	// wait for a period of time before blocking rendering
-	UE_LOG(LogAndroid, Log, TEXT("AndroidEGL::  SuspendApp_EventThread, waiting for event manager to process. tid: %d"), FPlatformTLS::GetCurrentThreadId());
+	UE_LOG(LogAndroid, Log, TEXT("SuspendApp_EventThread -> , waiting for event manager to process. tid: %d"), FPlatformTLS::GetCurrentThreadId());
 	EMDoneTrigger->Reset();
 	bool bSuccess = EMDoneTrigger->Wait(4000);
 	float ElapsedTimeInMs_EMDoneTrigger_Wait = FPlatformTime::ToMilliseconds(FPlatformTime::Cycles() - StartCycles);
-	UE_CLOG(!bSuccess, LogAndroid, Log, TEXT("backgrounding callback, not responded in timely manner."));
-	STANDALONE_DEBUG_LOGf(LogAndroid, TEXT("AndroidEGL::  SuspendApp_EventThread, EMDoneTrigger->Wait, waited '%f' ms"), (float)ElapsedTimeInMs_EMDoneTrigger_Wait);
+	UE_CLOG(!bSuccess, LogAndroid, Log, TEXT("SuspendApp_EventThread -> backgrounding callback, not responded in timely manner."));
+	STANDALONE_DEBUG_LOGf(LogAndroid, TEXT("SuspendApp_EventThread -> EMDoneTrigger->Wait, waited '%f' ms"), (float)ElapsedTimeInMs_EMDoneTrigger_Wait);
 
 	BlockRendering();
 
@@ -1690,20 +1685,12 @@ static void SuspendApp_EventThread()
 //Called from the event process thread
 static void OnAppCommandCB(struct android_app* app, int32_t cmd)
 {
-#if !USE_ANDROID_STANDALONE
 	check(IsInAndroidEventThread());
-#endif
 
 
 	// Set event thread's view of the window dimensions:
 	{
 		ANativeWindow* DimensionWindow = app->pendingWindow ? app->pendingWindow : app->window;
-#if USE_ANDROID_STANDALONE
-		if (GAndroidWindowOverride != NULL)
-		{
-			DimensionWindow = (ANativeWindow*)GAndroidWindowOverride;
-		}
-#endif
 		if (DimensionWindow)
 		{
 			FAndroidWindow::SetWindowDimensions_EventThread(DimensionWindow);
@@ -1735,23 +1722,29 @@ static void OnAppCommandCB(struct android_app* app, int32_t cmd)
 		DEVELOPER_LOG_COMMANDCB_CASE(APP_CMD_INIT_WINDOW);
 		STANDALONE_DEBUG_LOGf(LogAndroid, TEXT("Case APP_CMD_INIT_WINDOW"));
 		UE_LOG(LogAndroid, Log, TEXT("Case APP_CMD_INIT_WINDOW"));
-#if USE_ANDROID_STANDALONE
+		// give scope to local variable in case statement...
 		{
-			ANativeWindow* win = app->pendingWindow ? app->pendingWindow : app->window;
+			ANativeWindow* DimensionWindow = app->pendingWindow ? app->pendingWindow : app->window;
+#if USE_ANDROID_STANDALONE
 			if (GAndroidWindowOverride != NULL)
 			{
-				win = (ANativeWindow*)GAndroidWindowOverride;
+				DimensionWindow = (ANativeWindow*)GAndroidWindowOverride;
 			}
-			FAppEventManager::GetInstance()->HandleWindowCreated_EventThread((ANativeWindow*)win);
-		}
-#else
-		FAppEventManager::GetInstance()->HandleWindowCreated_EventThread(app->pendingWindow);
 #endif
-		bHasWindow = true;
+
+			FAppEventManager::GetInstance()->HandleWindowCreated_EventThread(DimensionWindow);
+			bHasWindow = DimensionWindow != nullptr;
+		}
+
+
+		UE_LOG(LogAndroid, Log, TEXT("Case APP_CMD_INIT_WINDOW: bHasWindow=%d, bHasFocus=%d, bIsResumed=%d"), bHasWindow, bHasFocus, bIsResumed);
+
+
 		if (bHasWindow && bHasFocus && bIsResumed)
 		{
 			ActivateApp_EventThread();
 		}
+
 		break;
 	case APP_CMD_TERM_WINDOW:
 		/**
@@ -1982,14 +1975,10 @@ static void OnAppCommandCB(struct android_app* app, int32_t cmd)
 
 
 		FAppEventManager::GetInstance()->EnqueueAppEvent(APP_EVENT_STATE_ON_DESTROY);
-#if !USE_ANDROID_STANDALONE
 
 		// Exit here, avoids having to unlock the window and letting the RHI's deal with invalid window.
 		extern void AndroidThunkCpp_ForceQuit();
 		AndroidThunkCpp_ForceQuit();
-#else
-		STANDALONE_DEBUG_LOGf(LogAndroid, TEXT("IGNORING ForceQuit in handling APP_CMD_DESTROY"));
-#endif
 
 		break;
 	}
@@ -2057,7 +2046,7 @@ JNI_METHOD void Java_com_epicgames_unreal_GameActivity_nativeConsoleCommand(JNIE
 	}
 	else
 	{
-		FPlatformMisc::LowLevelOutputDebugStringf(TEXT("Ignoring console command (too early): %s"), *Command);
+		FPlatformMisc::LowLevelOutputDebugStringf(TEXT("UNEXPECTED -> Ignoring console command (too early): %s"), *Command);
 	}
 #else
 	IssueConsoleCommand(Command);
@@ -2161,6 +2150,35 @@ static struct android_app dummy_state;
 #define EVENTTYPE_PRE_LOAD_MAP				5
 #define EVENTTYPE_POST_LOAD_MAP				6
 #define EVENTTYPE_ACTION					7
+#define EVENTTYPE_ACTIVITYCHANGED			8
+#define EVENTTYPE_APP_SUSPENDED				9
+
+
+static const TCHAR* eAppCommandNamed[] = {
+	TEXT("APP_CMD_INPUT_CHANGED"),
+	TEXT("APP_CMD_INIT_WINDOW"),
+	TEXT("APP_CMD_TERM_WINDOW"),
+	TEXT("APP_CMD_WINDOW_RESIZED"),
+	TEXT("APP_CMD_WINDOW_REDRAW_NEEDED"),
+	TEXT("APP_CMD_CONTENT_RECT_CHANGED"),
+	TEXT("APP_CMD_GAINED_FOCUS"),
+	TEXT("APP_CMD_LOST_FOCUS"),
+	TEXT("APP_CMD_CONFIG_CHANGED"),
+	TEXT("APP_CMD_LOW_MEMORY"),
+	TEXT("APP_CMD_START"),
+	TEXT("APP_CMD_RESUME"),
+	TEXT("APP_CMD_SAVE_STATE"),
+	TEXT("APP_CMD_PAUSE"),
+	TEXT("APP_CMD_STOP"),
+	TEXT("APP_CMD_DESTROY"),
+	TEXT("APP_CMD_UNKNOWN")
+};
+
+static const TCHAR* get_APP_CMD_String(uint32 cmd)
+{
+	check(cmd <= 15);
+	return (cmd <= 15) ? eAppCommandNamed[cmd] : eAppCommandNamed[16];
+}
 
 void AndroidThunkCpp_Engine_SendEvent(int32 event, const FString& param1, int32 param2, int32 param3, float param4)
 {
@@ -2193,6 +2211,12 @@ static void InitEvent()
 	AndroidThunkCpp_Engine_SendEvent(EVENTTYPE_INIT, TEXT(""), 0, 0, 0.0f);
 }
 
+static void AppSuspended()
+{
+	STANDALONE_DEBUG_LOGf(LogAndroid, TEXT("LaunchAndroid received AppSuspended"));
+	AndroidThunkCpp_Engine_SendEvent(EVENTTYPE_APP_SUSPENDED, TEXT(""), (int32)GFrameCounter, 0, 0.0f);
+}
+
 static void PostEngineInitEvent()
 {
 	STANDALONE_DEBUG_LOGf(LogAndroid, TEXT("LaunchAndroid received PostEngineInitEvent"));
@@ -2248,8 +2272,8 @@ JNI_METHOD void Java_com_epicgames_makeaar_GameActivityForMakeAAR_nativeMain(JNI
 	if (projectModule)
 	{
 		FString name = FJavaHelper::FStringFromParam(jenv, projectModule);
-		STANDALONE_DEBUG_LOGf(LogAndroid, TEXT("Java_com_epicgames_makeaar_GameActivityForMakeAAR_nativeMain : OverrideProjectModule to '%s'"), *name);
-		GEngineLoop.OverrideProjectModule(FJavaHelper::FStringFromParam(jenv, projectModule), name);
+		STANDALONE_DEBUG_LOGf(LogAndroid, TEXT("Java_com_epicgames_makeaar_GameActivityForMakeAAR_nativeMain : OverrideProjectModule to 'makeaar_GameActivityForMakeAAR', projectModule name is '%s'"), *name);
+		GEngineLoop.OverrideProjectModule(FJavaHelper::FStringFromParam(jenv, projectModule), TEXT("makeaar_GameActivityForMakeAAR"));
 	}
 	else
 	{
@@ -2393,48 +2417,19 @@ JNI_METHOD bool Java_com_epicgames_makeaar_Engine_nativeMountPak(JNIEnv* jenv, j
 
 JNI_METHOD void Java_com_epicgames_makeaar_GameActivityForMakeAAR_nativeAppCommand(JNIEnv* jenv, jobject thiz, jint cmd)
 {
+	STANDALONE_DEBUG_LOGf(LogAndroid, TEXT("MAKEAAR -> nativeAppCommand[%u, %s] -> EventThreadID=%d, GResumeMainInit=%d, GEventHandlerInitialized=%d, FTaskGraphInterface::IsRunning()=%d"), 
+		cmd,
+		get_APP_CMD_String(cmd),
+		EventThreadID, 
+		GResumeMainInit,
+		GEventHandlerInitialized,
+		FTaskGraphInterface::IsRunning()
+	);
+
 	dummy_state.window = (ANativeWindow*)GAndroidWindowOverride;
 	dummy_state.pendingWindow = (ANativeWindow*)GAndroidWindowOverride;
-	int useGameThreadValue = CVarUseGameThreadForCommands.GetValueOnAnyThread();
-	bool bForceCommandToGameThread = useGameThreadValue != 0;
-	STANDALONE_DEBUG_LOGf(LogAndroid, TEXT("nativeAppCommand -> useGameThreadValue=%d, FTaskGraphInterface::IsRunning()=%d, cmd: [%u],"), useGameThreadValue, FTaskGraphInterface::IsRunning(), cmd);
 
-	if (EventThreadID == 0)
-	{
-		STANDALONE_DEBUG_LOGf(LogAndroid, TEXT("waiting for EventThreadID != 0"));
-		while (EventThreadID == 0)
-		{
-			FPlatformProcess::Sleep(0.001f);
-			FPlatformMisc::MemoryBarrier();
-		};
-	}
-
-	if (useGameThreadValue == 0)
-	{
-		STANDALONE_DEBUG_LOGf(LogAndroid, TEXT("OnAppCommandCB(default) cmd: %u, tid = %d, FTaskGraphInterface::IsRunning()=%d"), cmd, gettid(), FTaskGraphInterface::IsRunning());
-		OnAppCommandCB(&dummy_state, cmd);
-	}
-	else if (bForceCommandToGameThread && FTaskGraphInterface::IsRunning())
-	{
-		STANDALONE_DEBUG_LOGf(LogAndroid, TEXT("OnAppCommandCB(Queue on GameThread) cmd: %u, tid = %d, FTaskGraphInterface::IsRunning()=%d"), cmd, gettid(), FTaskGraphInterface::IsRunning());
-		FGraphEventRef nativeAppCommandCalled = FFunctionGraphTask::CreateAndDispatchWhenReady([&]()
-			{
-				OnAppCommandCB(&dummy_state, cmd);
-			}, TStatId(), NULL, ENamedThreads::GameThread);
-	}
-	else
-	{
-		STANDALONE_DEBUG_LOGf(LogAndroid, TEXT("OnAppCommandCB(immediate) cmd: %d, tid = %d, FTaskGraphInterface::IsRunning()=%d"), (int)cmd, gettid(), FTaskGraphInterface::IsRunning());
-	    if (GEngine != NULL)
-	    {
-		    // Run on game thread to avoid race condition
-		    AsyncTask(ENamedThreads::GameThread, [cmd]()
-			    {
-				    STANDALONE_DEBUG_LOGf(LogAndroid, TEXT("Issue nativeAppCommand: %d, tid = %d"), (int)cmd, gettid());
-				    OnAppCommandCB(&dummy_state, cmd);
-			    });
-	    }
-	}
+	OnAppCommandCB(&dummy_state, cmd);
 }
 
 #endif // USE_ANDROID_STANDALONE
