Update AndroidSingleInstance plugin to properly generate ASISProject
#android

diff --git a/Engine/Plugins/Runtime/AndroidSingleInstanceService/Source/AndroidSingleInstanceService/Android/ASISStub/app/src/main/AndroidManifest.xml b/Engine/Plugins/Runtime/AndroidSingleInstanceService/Source/AndroidSingleInstanceService/Android/ASISStub/app/src/main/AndroidManifest.xml
new file mode 100644
--- /dev/null
+++ b/Engine/Plugins/Runtime/AndroidSingleInstanceService/Source/AndroidSingleInstanceService/Android/ASISStub/app/src/main/AndroidManifest.xml
@@ -0,0 +1,37 @@
+<?xml version="1.0" encoding="utf-8"?>
+<manifest xmlns:android="http://schemas.android.com/apk/res/android"
+	package="com.epicgames.ASISStub">
+	
+	<queries>
+		<package android:name="com.epicgames.Lyra"/>
+		<package android:name="com.epicgames.makeaar"/>
+		<package android:name="com.epicgames.unreal"/>
+	</queries>
+	
+	<application
+		android:label="@string/ASISAppName"
+		android:requestLegacyExternalStorage="true"
+		android:hasCode="true" 
+		android:usesCleartextTraffic="true" 
+		android:allowTaskReparenting="true" 
+		android:supportsRtl="true" 
+		android:hardwareAccelerated="true"
+		android:allowBackup="true"
+		android:theme="@style/FullscreenTheme"
+		>
+		<activity
+            android:name=".UseServiceActivity"
+			android:configChanges="orientation|keyboardHidden|screenSize"
+			android:hardwareAccelerated="true" android:allowTaskReparenting="true" android:exported="true"
+			>
+			<intent-filter>
+				<action android:name="android.intent.action.MAIN" />
+				<category android:name="android.intent.category.LAUNCHER" />
+			</intent-filter>
+		</activity>
+
+	</application>
+	
+ 
+
+</manifest>
diff --git a/Engine/Plugins/Runtime/AndroidSingleInstanceService/Source/AndroidSingleInstanceService/Android/ASISStub/app/src/main/java/com/epicgames/ASISStub/ASISConnection.java b/Engine/Plugins/Runtime/AndroidSingleInstanceService/Source/AndroidSingleInstanceService/Android/ASISStub/app/src/main/java/com/epicgames/ASISStub/ASISConnection.java
new file mode 100644
--- /dev/null
+++ b/Engine/Plugins/Runtime/AndroidSingleInstanceService/Source/AndroidSingleInstanceService/Android/ASISStub/app/src/main/java/com/epicgames/ASISStub/ASISConnection.java
@@ -0,0 +1,445 @@
+package com.epicgames.ASISStub;
+
+import android.app.Activity;
+import android.app.Application;
+import android.content.ComponentName;
+import android.content.Context;
+import android.content.ContextWrapper;
+import android.content.Intent;
+import android.content.ServiceConnection;
+import android.graphics.SurfaceTexture;
+import android.os.Build;
+import android.os.Bundle;
+import android.os.Handler;
+import android.os.IBinder;
+import android.os.Message;
+import android.os.Messenger;
+import android.os.RemoteException;
+import android.util.Log;
+import android.view.MotionEvent;
+import android.view.Surface;
+import android.view.TextureView;
+import android.view.View;
+
+import androidx.annotation.NonNull;
+import androidx.annotation.UiThread;
+import androidx.lifecycle.MutableLiveData;
+
+import com.epicgames.makeaar.UnrealMessageType;
+
+import java.util.concurrent.atomic.AtomicReference;
+
+
+public class ASISConnection implements ServiceConnection, TextureView.SurfaceTextureListener, TextureView.OnTouchListener
+{
+	final static String TAG = "UE-ASISConnection";
+
+	final static String obbModuleName = "Lyra";//""ALPHA_POC";
+	//	final static String obbModuleName = BuildConfig.ASISModuleName; //"Lyra";//""ALPHA_POC";
+//	final static String servicePackageName = "";
+//	final static String servicePackageName = BuildConfig.ASISPackageName;
+	final static String servicePackageName = "com.epicgames.Lyra";
+	final static String obbFileLocation = ""; //"/data/local/tmp/3d/AndroidPackagingTest.main.obb.png";
+//	final static String obbFileLocation = BuildConfig.ASISOBBPath; //"/data/local/tmp/3d/AndroidPackagingTest.main.obb.png";
+
+	final static String commandLineArgs = "-nosound";
+	//	final static String commandLineArgs = "-nosound -launchandroidflags=0 -dpcvars='Android.UseGameThreadForNativeCommands=0'";
+	final static boolean enablePropagateAlpha = true;
+	
+	Messenger serviceMessenger = null;
+	Messenger serviceReplyMessenger = null;
+	MutableLiveData<Boolean> isBoundToService = new MutableLiveData<Boolean>(false);
+
+
+	AtomicReference<Surface> externalSurfaceData = new AtomicReference<Surface>(null);
+
+	int[] surfacePosition = new int[2];
+
+	TextureView activeTextureView = null;
+
+	final int mTaskID;
+
+	private static final Handler.Callback mConnectionStateCallback = (Message msg) -> {
+		final String TAG = "UE-ServiceCallback";
+
+		switch (msg.what) {
+
+			case 0:	//EVENTTYPE_INIT == 0
+				Log.v(TAG, "received mConnectionStateCallback for EVENTTYPE_INIT, obj = " + msg.obj);
+
+				break;
+
+			case 1:	//EVENTTYPE_POST_ENGINE_INIT == 1
+				Log.v(TAG, "received mConnectionStateCallback for EVENTTYPE_POST_ENGINE_INIT, obj = " + msg.obj);
+				break;
+
+			default:
+				//Log.e(TAG, "mConnectionStateCallback, Unknown message = " + msg.what);
+		}
+		return true;
+	};
+	
+
+	int getTaskId() { return mTaskID; }
+
+	ContextWrapper mContextWrapper;
+
+	ASISConnection(Activity _activity)
+	{
+//			if (mServiceConnection.getValue() != null)
+//			{
+//				unbindToUnrealInstanceService("MyServiceConnection constructor!!!");
+//				doCleanupForUnbinding("MyServiceConnection constructor!!!");
+//
+//			}
+		assert (_activity != null);
+
+		mContextWrapper = _activity;
+
+		mTaskID = _activity.getTaskId();
+
+//			bindToUnrealInstanceService("MyServiceConnection constructor " + mTaskID);
+	}
+
+
+//		Activity getActivity() {
+//			assert (mContextWrapper instanceof Activity);
+//			return (Activity) mContextWrapper;
+//		}
+//
+//		Application getApplication()
+//		{
+//			Application application = null;
+//			try {
+//				application = getActivity().getApplication();
+//			}
+//			catch (Exception e)
+//			{
+//				e.printStackTrace();
+//			}
+//			finally {
+//				return application;
+//			}
+//		}
+
+	void SendMessage(Message msg) throws RemoteException {
+		if (serviceMessenger != null) {
+			serviceMessenger.send(msg);
+		}
+	}
+
+	void SetTextureView(TextureView textureView_)
+	{
+		textureView_.setAlpha(1.0f);
+		textureView_.setOpaque(false);
+		textureView_.setOnTouchListener(this);
+		textureView_.setSurfaceTextureListener(this);
+
+		Log.d(TAG, "SetTextureView: " + textureView_ + ", isAvailable=" + textureView_.isAvailable() + ", isAttachedToWindow=" + textureView_.isAttachedToWindow());
+
+		if (textureView_.isAvailable() && textureView_.isAttachedToWindow())
+		{
+			textureView_.getLocationOnScreen(surfacePosition);
+
+			onSurfaceTextureAvailable(textureView_.getSurfaceTexture(), textureView_.getWidth(), textureView_.getHeight());
+		}
+		activeTextureView = textureView_;
+	}
+
+	public boolean onTouch(View v, MotionEvent event) {
+		// TODO Auto-generated method stub
+
+		Log.d(TAG, "Touch: " + event.getAction() + " X: " + event.getX() + " Y: " + event.getY());
+
+		Message msg = Message.obtain(null, UnrealMessageType.TouchEvent.ordinal());
+		msg.getData().putInt("taskId", getTaskId());
+		msg.getData().putParcelable("touch", event);
+
+		try {
+			SendMessage(msg);
+		} catch (RemoteException e) {
+			e.printStackTrace();
+		}
+
+		return true;
+	}
+
+	public void ConsoleCommand(Context context, String consoleCommand)
+	{
+		Log.i(TAG, "ConsoleCommand() - consoleCommand = " + consoleCommand);
+
+		// Create and send a message to the service, using a supported 'what' value
+		Message msg = Message.obtain(null, UnrealMessageType.SendConsoleCommand.ordinal());
+		msg.getData().putString("consoleCommand", consoleCommand);
+
+		try {
+			SendMessage(msg);
+		} catch (RemoteException e) {
+			e.printStackTrace();
+		}
+	}
+
+	void SendData(int event, String param1, int param2, int param3, float param4)
+	{
+		// Create and send a message to the service, using a supported 'what' value
+		Message msg = Message.obtain(null, UnrealMessageType.SendData.ordinal());
+		msg.getData().putInt("taskId", getTaskId());
+		msg.getData().putInt("event", event);
+		msg.getData().putString("param1", param1);
+		msg.getData().putInt("param2", param2);
+		msg.getData().putInt("param3", param3);
+		msg.getData().putFloat("param4", param4);
+
+		try {
+			SendMessage(msg);
+		} catch (RemoteException e) {
+			e.printStackTrace();
+		}
+	}
+
+	void bindToUnrealInstanceService(String caller)
+	{
+		logContextDetails("bindToUnrealInstanceService("+ caller + ") -> ");
+		if(serviceReplyMessenger == null)
+		{
+			serviceReplyMessenger = new Messenger(new Handler(mConnectionStateCallback));
+		}
+
+		Intent intent = new Intent();
+		intent.setClassName(servicePackageName, "com.epicgames.makeaar.UnrealSharedInstanceService");
+
+		intent.putExtra("obbModuleName", obbModuleName);
+		intent.putExtra("obbFileLocation", obbFileLocation);
+		intent.putExtra("commandLineArgs", commandLineArgs);
+		intent.putExtra("enablePropagateAlpha", enablePropagateAlpha);
+
+		if (serviceReplyMessenger != null) {
+			intent.putExtra("callback", serviceReplyMessenger);
+		}
+
+		//ServiceConnection connection = GetServiceConnection("bindToUnrealInstanceService, caller=" + caller);
+
+		Log.d(TAG, "bindToUnrealInstanceService intent=" + intent + ", mServiceConnection=" + this);
+		mContextWrapper.bindService(intent, this, Context.BIND_AUTO_CREATE | Context.BIND_IMPORTANT | Context.BIND_DEBUG_UNBIND | Context.BIND_ADJUST_WITH_ACTIVITY);
+//			getApplication().bindService(intent, this, Context.BIND_AUTO_CREATE | Context.BIND_IMPORTANT | Context.BIND_ADJUST_WITH_ACTIVITY);
+	}
+
+	void unbindToUnrealInstanceService(String caller)
+	{
+		logContextDetails("unbindToUnrealInstanceService("+ caller + ") -> ");
+
+		if (isBoundToService.getValue()) {
+			Log.v(TAG, "unbindToUnrealInstanceService");
+
+			mContextWrapper.unbindService(this);
+		}
+	}
+
+
+	private void logContextDetails(String caller)
+	{
+		Log.w(TAG, caller + ": isBoundToService=" +isBoundToService.getValue()
+			+ ", taskID=" +getTaskId()
+			+ ", serviceMessenger=" +serviceMessenger
+			+ ", serviceReplyMessenger=" +serviceReplyMessenger
+			+ ", externalSurfaceData=" +externalSurfaceData
+		);
+	}
+
+	public void doCleanupForUnbinding(String caller) {
+		logContextDetails("doCleanupForUnbinding(" + caller + ")");
+		
+		if (activeTextureView != null) {
+			activeTextureView.setSurfaceTextureListener(null);
+			activeTextureView = null;
+		}
+
+		if (externalSurfaceData.get() != null)
+		{
+			externalSurfaceData.get().release();
+			externalSurfaceData.set(null);
+		}
+
+		serviceMessenger = null;
+		serviceReplyMessenger = null;
+		isBoundToService.setValue( false );
+	}
+
+	@Override
+	public void onServiceConnected(ComponentName name, IBinder serviceBinder) {
+		Log.v(TAG, "onServiceConnected name=" + name + ", serviceBinder=" + serviceBinder);
+		serviceMessenger = new Messenger(serviceBinder);
+		//serviceReplyMessenger = new Messenger(new Handler(mConnectionStateCallback));
+		isBoundToService.setValue( true );
+//			TextureView textureView_ = (TextureView) findViewById( R.id.renderView );
+//
+//			textureView_ = (TextureView) getActivity().findViewById( R.id.renderView );
+//			SetTextureView(textureView_);
+	}
+
+	@Override
+	public void onServiceDisconnected(ComponentName name) {
+
+		Log.w(TAG, "onServiceDisconnected name=" + name);
+
+		doCleanupForUnbinding("onServiceDisconnected name=" + name);
+	}
+
+	@Override
+	public void onBindingDied(ComponentName name) {
+
+		Log.w(TAG, "onBindingDied name=" + name);
+
+		doCleanupForUnbinding("onBindingDied name=" + name);
+
+		ServiceConnection.super.onBindingDied(name);
+	}
+
+	@Override
+	public void onNullBinding(ComponentName name) {
+
+		Log.w(TAG, "onNullBinding name=" + name);
+
+		//doCleanupForBinding("onNullBinding");
+
+		ServiceConnection.super.onNullBinding(name);
+	}
+
+	@Override
+	public void onSurfaceTextureAvailable(@NonNull SurfaceTexture surfaceTexture, int width, int height) {
+
+		logContextDetails("onSurfaceTextureAvailable(" + surfaceTexture + ")");
+
+		int[] surfacePosition = this.surfacePosition;
+		int[] surfaceSize = new int[]{width, height};
+
+		if (externalSurfaceData.get() != null) {
+			externalSurfaceData.get().release();
+		}
+
+		externalSurfaceData.set(new Surface(surfaceTexture));
+
+		attachSurfaceToService(getTaskId(), externalSurfaceData.get(), surfacePosition, surfaceSize);
+	}
+
+	@Override
+	public void onSurfaceTextureSizeChanged(@NonNull SurfaceTexture surfaceTexture, int width, int height) {
+		logContextDetails("onSurfaceTextureSizeChanged(" + surfaceTexture + ")");
+	}
+
+	// Invoked when the specified SurfaceTexture is about to be destroyed.
+	// If returns true, no rendering should happen inside the surface texture after this method is invoked.
+	// If returns false, the client needs to call SurfaceTexture.release().
+	// Most applications should return true.
+	@Override
+	public boolean onSurfaceTextureDestroyed(@NonNull SurfaceTexture surfaceTexture) {
+
+		synchronized (this) {
+			logContextDetails("onSurfaceTextureDestroyed(" + surfaceTexture + ")");
+			detachSurfaceFromService(getTaskId(), null);
+
+			if (externalSurfaceData.get() != null) {
+				externalSurfaceData.get().release();
+				externalSurfaceData.set(null);
+			}
+
+			if (activeTextureView != null) {
+				activeTextureView.setSurfaceTextureListener(null);
+			}
+			return true;
+		}
+	}
+
+	@Override
+	public void onSurfaceTextureUpdated(@NonNull SurfaceTexture surfaceTexture) {
+	}
+
+	@UiThread
+	void attachSurfaceToService(int attachId, Surface externalSurface, int[] viewPosition, int[] viewSize)
+	{
+		synchronized (this) {
+			if (externalSurface == null) {
+				Log.w(TAG, "attachSurfaceToService, externalSurface is null");
+				return;
+			}
+
+			if (Build.VERSION.SDK_INT >= Build.VERSION_CODES.P) {
+				Log.v(TAG, "attachSurfaceToService() - " + "proc = " + Application.getProcessName() + ", externalSurface = " + externalSurface + ", attachId = " + attachId);
+			}
+
+			Message msg = Message.obtain(null, UnrealMessageType.AttachExternalSurface.ordinal());
+			Bundle data = msg.getData();
+			data.putInt("attachId", attachId);
+			data.putBoolean("enablePropagateAlpha", true);
+			data.putBoolean("start", true);
+			data.putBoolean("resume", true);
+
+			data.putParcelable("surface", externalSurface);
+
+			if (viewPosition.length == 2) {
+				data.putIntArray("viewPosition", viewPosition);
+			}
+
+			if (viewSize.length == 2) {
+				data.putIntArray("viewSize", viewSize);
+
+			}
+
+			try {
+				SendMessage(msg);
+			} catch (RemoteException e) {
+				e.printStackTrace();
+			}
+		}
+	}
+
+	@UiThread
+	void detachSurfaceFromService(int attachId, Handler.Callback callback) 
+	{
+
+		synchronized (this) {
+			Message msg = Message.obtain(
+				null,
+				UnrealMessageType.DetachExternalSurface.ordinal()
+			);
+
+			msg.getData().putInt("attachId", attachId);
+			msg.getData().putBoolean("pause", true);
+			msg.getData().putBoolean("stop", true);
+
+			if (callback != null) {
+
+				Handler handler = new Handler(callback);
+				Messenger replyMessenger = new Messenger(handler);
+
+				msg.replyTo = replyMessenger;
+			}
+
+			Log.i(TAG, "detachSurfaceFromService() - " + "proc = ${Application.getProcessName()}");
+
+			try {
+				SendMessage(msg);
+			} catch (RemoteException e) {
+				e.printStackTrace();
+			}
+		}
+	}
+
+	@UiThread
+	void resumeService(int attachId)
+	{
+		synchronized (this) {
+			if (isBoundToService.getValue()) {
+				Message msg = Message.obtain(null, UnrealMessageType.ResumeService.ordinal());
+				Bundle data = msg.getData();
+				data.putInt("attachId", attachId);
+
+				try {
+					SendMessage(msg);
+				} catch (RemoteException e) {
+					e.printStackTrace();
+				}
+			}
+		}
+	}
+}
\ No newline at end of file
diff --git a/Engine/Plugins/Runtime/AndroidSingleInstanceService/Source/AndroidSingleInstanceService/Android/ASISStub/app/src/main/java/com/epicgames/ASISStub/UseServiceActivity.java b/Engine/Plugins/Runtime/AndroidSingleInstanceService/Source/AndroidSingleInstanceService/Android/ASISStub/app/src/main/java/com/epicgames/ASISStub/UseServiceActivity.java
new file mode 100644
--- /dev/null
+++ b/Engine/Plugins/Runtime/AndroidSingleInstanceService/Source/AndroidSingleInstanceService/Android/ASISStub/app/src/main/java/com/epicgames/ASISStub/UseServiceActivity.java
@@ -0,0 +1,473 @@
+package com.epicgames.ASISStub;
+
+import androidx.annotation.NonNull;
+import androidx.annotation.UiThread;
+import androidx.appcompat.app.AppCompatDelegate;
+import androidx.lifecycle.MutableLiveData;
+
+import android.app.Activity;
+import android.app.Application;
+import android.content.ComponentName;
+import android.content.Context;
+import android.content.ContextWrapper;
+import android.content.ServiceConnection;
+import android.content.res.Configuration;
+import android.graphics.SurfaceTexture;
+import android.os.Build;
+import android.os.Bundle;
+import android.os.Debug;
+import android.os.Handler;
+import android.os.IBinder;
+import android.os.Looper;
+import android.os.Message;
+import android.os.Messenger;
+import android.os.RemoteException;
+import android.util.Log;
+import android.view.MotionEvent;
+import android.view.Surface;
+import android.view.TextureView;
+import android.view.View;
+import android.content.Intent;
+import android.view.Window;
+import android.view.WindowManager;
+import android.widget.Button;
+import android.widget.RelativeLayout;
+
+import com.epicgames.makeaar.UnrealMessageType;
+
+import java.lang.ref.WeakReference;
+import java.util.concurrent.atomic.AtomicReference;
+
+
+public class UseServiceActivity extends Activity implements View.OnClickListener
+{
+
+	Messenger mMessenger;
+//	private static WeakReference<View.OnClickListener> mCallbackRef = new WeakReference<>(IEventCallback.NO_OP);
+
+
+
+	//final MutableLiveData<MyServiceConnection> mServiceConnection = new MutableLiveData<MyServiceConnection>(new MyServiceConnection());
+
+	ASISConnection mServiceConnection = null; //new MyServiceConnection();
+	
+	ASISConnection GetServiceConnection(String caller)
+	{
+
+		Log.w(TAG, "GetServiceConnection(" + caller + ") : mServiceConnection = " + mServiceConnection);
+
+		if (mServiceConnection != null)
+		{
+			if (!mServiceConnection.isBoundToService.getValue())
+			{
+				mServiceConnection.bindToUnrealInstanceService("GetServiceConnection rebind");
+			}
+			return mServiceConnection;
+		}
+		else
+		{
+			CreateServiceConnection(this, "GetServiceConnection");
+		}
+		
+		assert (mServiceConnection != null);
+
+		return mServiceConnection;
+	}
+
+	ASISConnection CreateServiceConnection(Activity activity, String caller)
+	{
+
+		Log.w(TAG, "CreateServiceConnection(" + caller + ") : mServiceConnection = " + mServiceConnection + ", activity=" + activity);
+
+		if (mServiceConnection != null)
+		{
+			if (!mServiceConnection.isBoundToService.getValue())
+			{
+				mServiceConnection.bindToUnrealInstanceService("CreateServiceConnection rebind");
+			}
+			return mServiceConnection;
+		}
+		
+//		if (mServiceConnection != null && mServiceConnection.isBoundToService.getValue())
+//		{
+//			return mServiceConnection;
+////			mServiceConnection.logContextDetails("CreateServiceConnection destroy prev (" + activity + ")");
+////			DeposeServiceConnection();
+//		}
+
+		if (mServiceConnection == null)
+		{
+			//mServiceConnection.logContextDetails("CreateServiceConnection create new (" + activity + ")");
+			mServiceConnection = new ASISConnection(activity);
+			mServiceConnection.bindToUnrealInstanceService("CreateServiceConnection");
+		}
+
+		return mServiceConnection;
+	}
+
+	void DeposeServiceConnection(String caller)
+	{
+
+		Log.w(TAG, "DeposeServiceConnection, mServiceConnection = " + mServiceConnection + ", caller=" + caller);
+
+		if (mServiceConnection != null && mServiceConnection.isBoundToService.getValue()) {
+			mServiceConnection.unbindToUnrealInstanceService(caller + "->DeposeServiceConnection");
+			mServiceConnection.doCleanupForUnbinding(caller + "->DeposeServiceConnection");
+		}
+
+		mServiceConnection = null;
+	}
+
+
+
+	static Integer next_theme = AppCompatDelegate.MODE_NIGHT_NO;
+
+	// declaring objects of Button class
+	// Surface renderSurface;
+
+	final static String TAG = "UE-UseServiceActivity";
+
+
+	
+	private Button activateView1Button, activateView2Button, activateView3Button, restartActivityButton, themeChangeButton, resumeButton;
+	private Button pauseButton, stopButton, bindServiceButton, unbindServiceButton, memReportButton;
+
+	void ActivateTextureView(int viewIndex)
+	{
+		TextureView textureView_ = null;
+		switch (viewIndex)
+		{
+			case 1:
+				textureView_ = (TextureView) findViewById( R.id.textureView );
+				break;
+			case 2:
+				textureView_ = (TextureView) findViewById( R.id.renderView );
+				break;
+			case 3:
+				textureView_ = (TextureView) findViewById( R.id.renderView2 );
+				break;
+			default:
+				assert false;
+		}
+
+//		try {
+//			GetServiceConnection("ActivateTextureView cleanup").detachSurfaceFromService(getTaskId(), null);
+//		} catch (Exception e) {
+//			e.printStackTrace();
+//		}
+
+		GetServiceConnection("ActivateTextureView").SetTextureView(textureView_);
+	}
+
+
+	@UiThread
+	void stopService(int attachId, Handler.Callback callback)
+	{
+		Intent intent = new Intent();
+		intent.setClassName(ASISConnection.servicePackageName, "com.epicgames.makeaar.UnrealSharedInstanceService");
+		getApplication().stopService(intent);
+	}
+
+	@Override
+	public void onContentChanged() {
+		Log.v(TAG, "UseServiceActivity: onContentChanged");
+		super.onContentChanged();
+	}
+
+	@Override
+	protected void onResume() {
+		Log.v(TAG, "UseServiceActivity: onResume");
+		super.onResume();
+
+		//GetServiceConnection("onResume").resumeService(getTaskId());
+	}
+
+	@Override
+	protected void onPause() {
+		Log.v(TAG, "UseServiceActivity: onPause");
+		super.onPause();
+		GetServiceConnection("onPause").detachSurfaceFromService(getTaskId(), null);
+	}
+
+	@Override
+	protected void onStart() {
+		Log.v(TAG, "UseServiceActivity: onStart");
+		super.onStart();
+
+//		CreateServiceConnection( this,"UseServiceActivity::onStart()");
+	}
+
+	@Override
+	protected void onStop() {
+		Log.v(TAG, "UseServiceActivity: onStop");
+		super.onStop();
+//		DeposeServiceConnection();
+	}
+
+	@Override
+	protected void onDestroy() {
+		Log.v(TAG, "UseServiceActivity: onDestroy");
+		super.onDestroy();
+
+		DeposeServiceConnection("onDestroy");
+	}
+
+	@Override
+	protected void onRestart() {
+		Log.v(TAG, "UseServiceActivity onRestart called! hasBeenCreated = " + hasBeenCreated);
+		super.onRestart();
+	}
+
+	
+	@Override
+	protected void onCreate(Bundle savedInstanceState) {
+
+		new RelativeLayout(this);
+		
+//		getWindow().getTransitionManager().setTransition();
+//		WindowManager.LayoutParams params = getWindow().getAttributes();
+//		params.type = params.TYPE_APPLICATION_ATTACHED_DIALOG;
+//		params.flags = params.flags | Window.FEATURE_NO_TITLE;
+
+//!		getWindow();
+//!		getWindowManager();
+//!		getIntent();
+//!		runOnUiThread();
+		
+//		getResources();
+//		getExternalFilesDir();
+//		getFilesDir();
+//		getSystemService();
+//		getPackageName();
+//		getPackageResourcePath();
+//		getMainExecutor();
+//		getDisplay();
+//		getPackageManager();
+//		getApplicationContext();
+
+
+		
+		super.onCreate( savedInstanceState );
+		AppCompatDelegate.setDefaultNightMode(next_theme);
+		getTheme().applyStyle(R.style.Theme_UnrealEngine_NoActionBar, true);
+		
+		setContentView( R.layout.interface_use_service_activity );
+
+
+
+		activateView1Button 	= (Button) findViewById( R.id.activateView1Button);
+		activateView2Button 	= (Button) findViewById( R.id.activateView2Button);
+		activateView3Button 	= (Button) findViewById( R.id.activateView3Button);
+
+		pauseButton 			= (Button) findViewById( R.id.buttonPause);
+		stopButton 				= (Button) findViewById( R.id.buttonStopService);
+		bindServiceButton 		= (Button) findViewById( R.id.buttonBindService);
+		unbindServiceButton 	= (Button) findViewById( R.id.buttonUnbindService);
+		restartActivityButton 	= (Button) findViewById( R.id.restartActivityButton);
+		memReportButton 		= (Button) findViewById( R.id.buttonMemReport);
+		themeChangeButton 		= (Button) findViewById( R.id.changeThemeButton);
+		resumeButton 			= (Button) findViewById( R.id.resumeButton);
+
+
+		activateView1Button.setOnClickListener( this );
+		activateView2Button.setOnClickListener( this );
+		activateView3Button.setOnClickListener( this );
+
+		pauseButton.setOnClickListener( this );
+		stopButton.setOnClickListener( this );
+		bindServiceButton.setOnClickListener( this );
+		unbindServiceButton.setOnClickListener( this );
+		themeChangeButton.setOnClickListener( this );
+		resumeButton.setOnClickListener( this );
+
+
+		Log.v(TAG, "UseServiceActivity onCreate called! hasBeenCreated = " + hasBeenCreated);
+
+		if (false == hasBeenCreated) {
+			Debug.getMemoryInfo(baseMemoryInfo);
+			Debug.getMemoryInfo(prevMemoryInfo);
+			hasBeenCreated = true;
+		}
+
+		CreateServiceConnection( this,"UseServiceActivity::onCreate()");
+	}
+
+	public void onClick(View view)
+	{
+		final boolean waitForBind = false;
+
+		if (activateView1Button.equals(view)) {
+			ActivateTextureView(1);
+		} else if (activateView2Button.equals(view)) {
+			ActivateTextureView(2);
+		} else if (activateView3Button.equals(view)) {
+			ActivateTextureView(3);
+		} else if (pauseButton.equals(view)) {
+
+			try {
+				GetServiceConnection("Pause Button").detachSurfaceFromService(getTaskId(), null);
+			} catch (Exception e) {
+				e.printStackTrace();
+			}
+
+		} else if (stopButton.equals(view)) {
+			stopService(getTaskId(), null);
+		} else if (bindServiceButton.equals(view)) {
+
+			try {
+				CreateServiceConnection( this,"Bind Service Button").bindToUnrealInstanceService("onClick");
+			} catch (Exception e) {
+				e.printStackTrace();
+			}
+
+		} else if (unbindServiceButton.equals(view)) {
+
+			try {
+				DeposeServiceConnection("unbindServiceButton");
+				//GetServiceConnection("Unbind Service Button").unbindToUnrealInstanceService("onClick");
+			} catch (Exception e) {
+				e.printStackTrace();
+			}
+
+		} else if (restartActivityButton.equals(view)) {
+
+			processRestartActivity();
+		} else if (memReportButton.equals(view)) {
+			ReportDebugMemoryInfo();
+		} else if (themeChangeButton.equals(view)) {
+			handleThemeChange();
+		} else if (resumeButton.equals(view)) {
+
+			try {
+				GetServiceConnection("Resume Button").resumeService(getTaskId());
+			} catch (Exception e) {
+				e.printStackTrace();
+			}
+		}
+
+	}
+
+
+
+//    @Override
+//  public boolean onTouchEvent(MotionEvent e) {
+//      return ueInterface.MessageServiceToHandleTouch(e);
+//    }
+
+
+	//private final UEInterfaceHelpers ueInterface = new UEInterfaceHelpers(this);
+
+	int cycleThemeCounter = 0;
+	static boolean hasBeenCreated = false;
+	static Debug.MemoryInfo memoryInfo = new Debug.MemoryInfo();
+	static Debug.MemoryInfo prevMemoryInfo = new Debug.MemoryInfo();
+	static Debug.MemoryInfo baseMemoryInfo = new Debug.MemoryInfo();
+	static int tryCount = 0;
+	void ReportDebugMemoryInfo()
+	{
+//        if (prevMemoryInfo != null)
+//        {
+//            Log.i(TAG, "Debug.memoryInfo prev= " + prevMemoryInfo.getMemoryStats());
+//
+//    }
+		Runtime.getRuntime().gc();
+
+		Debug.getMemoryInfo(memoryInfo);
+		Log.v(TAG, "Debug.memoryInfo:\n\tthis= " + memoryInfo.getMemoryStats()
+				+ "\n\tprev= " + prevMemoryInfo.getMemoryStats()
+				+ "\n\tbase= " + baseMemoryInfo.getMemoryStats()
+		);
+		Debug.getMemoryInfo(prevMemoryInfo);
+
+	}
+
+	Handler.Callback onDetachSurfaceHandler = new Handler.Callback() {
+		@Override
+		public boolean handleMessage(Message msg) {
+			synchronized (mServiceConnection) {
+				Log.v(TAG, "onDetachSurfaceHandler::handleMessage() msg=" + msg);
+
+				recreate();
+				return true;
+			}
+		}
+	};
+	
+	private Handler handler = new Handler();
+
+	private void processRestartActivity() {
+
+		mServiceConnection.detachSurfaceFromService(getTaskId(), onDetachSurfaceHandler);
+		//recreate();
+		
+//		Runnable RestartTrigger = new Runnable() {
+//			@Override
+//			public void run() {
+//
+//				mServiceConnection.detachSurfaceFromService(getTaskId(), onDetachSurfaceHandler);
+//			}
+//		};
+//
+//		handler.post(RestartTrigger);
+	}
+
+	void doThemeChange(Integer value)
+	{
+		Log.v(TAG, "doing a touch event theme change: new uiTheme = ${value}");
+//        prev_uiMode = next_theme;
+		next_theme = value;
+//        AppCompatDelegate.setDefaultNightMode(value);
+		if (mServiceConnection != null)
+		{
+			mServiceConnection.detachSurfaceFromService(getTaskId(), onDetachSurfaceHandler);
+		}
+		else {
+			recreate();
+		}
+	}
+
+	@Override
+	public void onConfigurationChanged(@NonNull Configuration newConfig) {
+		super.onConfigurationChanged(newConfig);
+	}
+
+	void handleThemeChange()
+	{
+		//GetServiceConnection().bindToUnrealInstanceService("handleThemeChange");
+		//resumeService(getTaskId(), null);
+
+		doThemeChange(next_theme == AppCompatDelegate.MODE_NIGHT_NO ? AppCompatDelegate.MODE_NIGHT_YES : AppCompatDelegate.MODE_NIGHT_NO);
+
+
+		tryCount = 110;
+//            ueInterface.bindToUnrealInstanceService(this, textureView_, BuildConfig.OBB_FILE_NAME, BuildConfig.OBB_MODULE_NAME, "-nosound -nocrashreports -opengl -featureleveles31", true, mConnectionStateCallback);
+//            ReportDebugMemoryInfo();
+
+
+//            ueInterface.runOnMainThread(new Runnable()
+		runOnUiThread(new Runnable()
+		{
+			@Override
+			public void run() {
+				Context context = getApplicationContext();
+				synchronized (context) {
+
+					try {
+						while ( (++tryCount) < 15) {
+							Log.v(TAG, "mem pressure test iter = " + tryCount);
+
+
+							doThemeChange(next_theme == AppCompatDelegate.MODE_NIGHT_NO ? AppCompatDelegate.MODE_NIGHT_YES : AppCompatDelegate.MODE_NIGHT_NO);
+
+						}
+					} catch (Exception e) {
+						e.printStackTrace();
+					}
+
+				}
+			}
+		});
+	}
+
+
+}
diff --git a/Engine/Plugins/Runtime/AndroidSingleInstanceService/Source/AndroidSingleInstanceService/Android/ASISStub/app/src/main/java/com/epicgames/makeaar/UnrealMessageType.java b/Engine/Plugins/Runtime/AndroidSingleInstanceService/Source/AndroidSingleInstanceService/Android/ASISStub/app/src/main/java/com/epicgames/makeaar/UnrealMessageType.java
new file mode 100644
--- /dev/null
+++ b/Engine/Plugins/Runtime/AndroidSingleInstanceService/Source/AndroidSingleInstanceService/Android/ASISStub/app/src/main/java/com/epicgames/makeaar/UnrealMessageType.java
@@ -0,0 +1,13 @@
+package com.epicgames.makeaar;
+
+public enum UnrealMessageType {
+	Undefined,  // use up 0
+	Hello,
+	AttachExternalSurface,
+	DetachExternalSurface,
+	StopService,
+	ResumeService,
+	TouchEvent,
+	SendConsoleCommand,
+	SendData
+}
\ No newline at end of file
diff --git a/Engine/Plugins/Runtime/AndroidSingleInstanceService/Source/AndroidSingleInstanceService/Android/ASISStub/app/src/main/res/layout/interface_use_service_activity.xml b/Engine/Plugins/Runtime/AndroidSingleInstanceService/Source/AndroidSingleInstanceService/Android/ASISStub/app/src/main/res/layout/interface_use_service_activity.xml
new file mode 100644
--- /dev/null
+++ b/Engine/Plugins/Runtime/AndroidSingleInstanceService/Source/AndroidSingleInstanceService/Android/ASISStub/app/src/main/res/layout/interface_use_service_activity.xml
@@ -0,0 +1,177 @@
+<?xml version="1.0" encoding="utf-8"?>
+<androidx.constraintlayout.widget.ConstraintLayout xmlns:android="http://schemas.android.com/apk/res/android"
+	xmlns:app="http://schemas.android.com/apk/res-auto"
+	xmlns:tools="http://schemas.android.com/tools"
+	android:layout_width="match_parent"
+	android:layout_height="match_parent"
+	android:layout_gravity="start"
+	android:visibility="visible"
+	tools:layout_editor_absoluteY="89dp">
+
+	<LinearLayout
+		android:id="@+id/linearLayout"
+		android:layout_width="match_parent"
+		android:layout_height="wrap_content"
+		android:orientation="horizontal"
+		tools:ignore="MissingConstraints">
+
+		<Button
+			android:id="@+id/buttonPause"
+			android:layout_width="wrap_content"
+			android:layout_height="wrap_content"
+			android:layout_weight="1"
+			android:onClick="onClick"
+			android:text="Pause"
+			android:visibility="visible" />
+
+		<Button
+			android:id="@+id/buttonBindService"
+			android:layout_width="wrap_content"
+			android:layout_height="wrap_content"
+			android:layout_weight="1"
+			android:minWidth="48dp"
+			android:onClick="onClick"
+			android:text="Bind Service"
+			android:visibility="visible" />
+
+		<Button
+			android:id="@+id/buttonUnbindService"
+			android:layout_width="wrap_content"
+			android:layout_height="wrap_content"
+			android:layout_weight="1"
+			android:onClick="onClick"
+			android:text="Unbind Service"
+			android:visibility="visible" />
+
+		<Button
+			android:id="@+id/buttonStopService"
+			android:layout_width="wrap_content"
+			android:layout_height="wrap_content"
+			android:layout_weight="1"
+			android:onClick="onClick"
+			android:text="Stop Service"
+			android:visibility="visible" />
+
+		<Button
+			android:id="@+id/buttonMemReport"
+			android:layout_width="wrap_content"
+			android:layout_height="wrap_content"
+			android:layout_weight="1"
+			android:onClick="onClick"
+			android:text="Mem Report"
+			android:visibility="visible" />
+
+	</LinearLayout>
+
+	<LinearLayout
+		android:id="@+id/layoutButtons"
+		android:layout_width="match_parent"
+		android:layout_height="wrap_content"
+		android:layout_marginTop="0dp"
+		android:orientation="horizontal"
+		android:translationY="0dp"
+		app:layout_constraintTop_toBottomOf="@+id/linearLayout">
+
+		<Button
+			android:id="@+id/activateView1Button"
+			android:layout_width="wrap_content"
+			android:layout_height="wrap_content"
+			android:onClick="onClick"
+			android:text="Activate View1"
+			android:visibility="visible" />
+
+		<Button
+			android:id="@+id/activateView2Button"
+			android:layout_width="wrap_content"
+			android:layout_height="wrap_content"
+			android:onClick="onClick"
+			android:text="Activate View2"
+			android:visibility="visible" />
+
+		<Button
+			android:id="@+id/activateView3Button"
+			android:layout_width="wrap_content"
+			android:layout_height="wrap_content"
+			android:onClick="onClick"
+			android:text="Activate View3"
+			android:visibility="visible" />
+
+		<Button
+			android:id="@+id/restartActivityButton"
+			android:layout_width="wrap_content"
+			android:layout_height="wrap_content"
+			android:layout_weight="1"
+			android:minWidth="32dp"
+			android:onClick="onClick"
+			android:text="RestartActivity"
+			android:visibility="visible" />
+
+		<Button
+			android:id="@+id/changeThemeButton"
+			android:layout_width="wrap_content"
+			android:layout_height="wrap_content"
+			android:layout_weight="1"
+			android:onClick="onClick"
+			android:text="changeTheme"
+			android:visibility="visible" />
+
+		<Button
+			android:id="@+id/resumeButton"
+			android:layout_width="wrap_content"
+			android:layout_height="wrap_content"
+			android:layout_weight="1"
+			android:onClick="onClick"
+			android:text="Resume Service"
+			android:visibility="visible" />
+
+	</LinearLayout>
+
+	<HorizontalScrollView
+		android:id="@+id/scrollViewMiddle"
+		android:contentDescription="Render Biews"
+		android:layout_width="wrap_content"
+		android:layout_height="wrap_content"
+		android:layout_marginTop="100dp"
+		android:background="#2196F3"
+		app:layout_constraintTop_toTopOf="parent"
+		tools:ignore="MissingConstraints">
+
+		<LinearLayout
+			android:layout_width="wrap_content"
+			android:layout_height="400dp"
+			android:background="#CDDC39"
+			android:orientation="horizontal">
+
+			<TextureView
+				android:id="@+id/textureView"
+				android:layout_width="300dp"
+				android:layout_height="match_parent" />
+
+			<View
+				android:id="@+id/divider2"
+				android:layout_width="1dp"
+				android:layout_height="match_parent"
+				android:background="?android:attr/listDivider" />
+
+			<TextureView
+				android:id="@+id/renderView"
+				android:layout_width="1000dp"
+				android:layout_height="match_parent" />
+
+			<View
+				android:id="@+id/divider"
+				android:layout_width="1dp"
+				android:layout_height="match_parent"
+				android:layout_weight="1"
+				android:background="?android:attr/listDivider" />
+
+			<TextureView
+				android:id="@+id/renderView2"
+				android:layout_width="800dp"
+				android:layout_height="match_parent"
+				android:layout_weight="1" />
+
+		</LinearLayout>
+	</HorizontalScrollView>
+
+</androidx.constraintlayout.widget.ConstraintLayout>
\ No newline at end of file
diff --git a/Engine/Plugins/Runtime/AndroidSingleInstanceService/Source/AndroidSingleInstanceService/Android/ASISStub/app/src/main/res/values/attrs.xml b/Engine/Plugins/Runtime/AndroidSingleInstanceService/Source/AndroidSingleInstanceService/Android/ASISStub/app/src/main/res/values/attrs.xml
new file mode 100644
--- /dev/null
+++ b/Engine/Plugins/Runtime/AndroidSingleInstanceService/Source/AndroidSingleInstanceService/Android/ASISStub/app/src/main/res/values/attrs.xml
@@ -0,0 +1,12 @@
+<resources>
+
+	<!-- Declare custom theme attributes that allow changing which styles are
+         used for button bars depending on the API level.
+         ?android:attr/buttonBarStyle is new as of API 11 so this is
+         necessary to support previous API levels. -->
+	<declare-styleable name="ButtonBarContainerTheme">
+		<attr name="metaButtonBarStyle" format="reference" />
+		<attr name="metaButtonBarButtonStyle" format="reference" />
+	</declare-styleable>
+
+</resources>
diff --git a/Engine/Plugins/Runtime/AndroidSingleInstanceService/Source/AndroidSingleInstanceService/Android/ASISStub/app/src/main/res/values/colors.xml b/Engine/Plugins/Runtime/AndroidSingleInstanceService/Source/AndroidSingleInstanceService/Android/ASISStub/app/src/main/res/values/colors.xml
new file mode 100644
--- /dev/null
+++ b/Engine/Plugins/Runtime/AndroidSingleInstanceService/Source/AndroidSingleInstanceService/Android/ASISStub/app/src/main/res/values/colors.xml
@@ -0,0 +1,12 @@
+<?xml version="1.0" encoding="utf-8"?>
+<resources>
+    <color name="colorPrimary">#6200EE</color>
+    <color name="colorPrimaryDark">#3700B3</color>
+    <color name="colorAccent">#03DAC5</color>
+
+    <color name="black_overlay">#66000000</color>
+    <color name="light_blue_400">#FF29B6F6</color>
+    <color name="light_blue_600">#FF039BE5</color>
+    <color name="gray_400">#FFBDBDBD</color>
+    <color name="gray_600">#FF757575</color>
+</resources>
diff --git a/Engine/Plugins/Runtime/AndroidSingleInstanceService/Source/AndroidSingleInstanceService/Android/ASISStub/app/src/main/res/values/dimens.xml b/Engine/Plugins/Runtime/AndroidSingleInstanceService/Source/AndroidSingleInstanceService/Android/ASISStub/app/src/main/res/values/dimens.xml
new file mode 100644
--- /dev/null
+++ b/Engine/Plugins/Runtime/AndroidSingleInstanceService/Source/AndroidSingleInstanceService/Android/ASISStub/app/src/main/res/values/dimens.xml
@@ -0,0 +1,5 @@
+<resources>
+    <dimen name="app_bar_height">180dp</dimen>
+    <dimen name="fab_margin">16dp</dimen>
+    <dimen name="text_margin">16dp</dimen>
+</resources>
\ No newline at end of file
diff --git a/Engine/Plugins/Runtime/AndroidSingleInstanceService/Source/AndroidSingleInstanceService/Android/ASISStub/app/src/main/res/values/strings.xml b/Engine/Plugins/Runtime/AndroidSingleInstanceService/Source/AndroidSingleInstanceService/Android/ASISStub/app/src/main/res/values/strings.xml
new file mode 100644
--- /dev/null
+++ b/Engine/Plugins/Runtime/AndroidSingleInstanceService/Source/AndroidSingleInstanceService/Android/ASISStub/app/src/main/res/values/strings.xml
@@ -0,0 +1,3 @@
+<resources>
+    <string name="app_name">ASISProject</string>
+</resources>
diff --git a/Engine/Plugins/Runtime/AndroidSingleInstanceService/Source/AndroidSingleInstanceService/Android/ASISStub/app/src/main/res/values/styles.xml b/Engine/Plugins/Runtime/AndroidSingleInstanceService/Source/AndroidSingleInstanceService/Android/ASISStub/app/src/main/res/values/styles.xml
new file mode 100644
--- /dev/null
+++ b/Engine/Plugins/Runtime/AndroidSingleInstanceService/Source/AndroidSingleInstanceService/Android/ASISStub/app/src/main/res/values/styles.xml
@@ -0,0 +1,23 @@
+<resources>
+
+    <!-- Base application theme. -->
+    <style name="AppTheme" parent="Theme.AppCompat.Light.DarkActionBar">
+        <!-- Customize your theme here. -->
+        <item name="colorPrimary">@color/colorPrimary</item>
+        <item name="colorPrimaryDark">@color/colorPrimaryDark</item>
+        <item name="colorAccent">@color/colorAccent</item>
+    </style>
+
+    <style name="FullscreenTheme" parent="AppTheme">
+        <item name="android:actionBarStyle">@style/FullscreenActionBarStyle</item>
+        <item name="android:windowActionBarOverlay">true</item>
+        <item name="android:windowBackground">@null</item>
+        <item name="metaButtonBarStyle">?android:attr/buttonBarStyle</item>
+        <item name="metaButtonBarButtonStyle">?android:attr/buttonBarButtonStyle</item>
+    </style>
+
+    <style name="FullscreenActionBarStyle" parent="Widget.AppCompat.ActionBar">
+        <item name="android:background">@color/black_overlay</item>
+    </style>
+
+</resources>
diff --git a/Engine/Plugins/Runtime/AndroidSingleInstanceService/Source/AndroidSingleInstanceService/Android/ASISStub/app/src/main/res/values/themes.xml b/Engine/Plugins/Runtime/AndroidSingleInstanceService/Source/AndroidSingleInstanceService/Android/ASISStub/app/src/main/res/values/themes.xml
new file mode 100644
--- /dev/null
+++ b/Engine/Plugins/Runtime/AndroidSingleInstanceService/Source/AndroidSingleInstanceService/Android/ASISStub/app/src/main/res/values/themes.xml
@@ -0,0 +1,8 @@
+<resources>
+
+    <style name="Theme.UnrealEngine.NoActionBar" parent="Theme.AppCompat">
+        <item name="windowActionBar">false</item>
+        <item name="windowNoTitle">true</item>
+    </style>
+	
+</resources>
\ No newline at end of file
diff --git a/Engine/Plugins/Runtime/AndroidSingleInstanceService/Source/AndroidSingleInstanceService/Android/ASISStub/src/main/java/com/epicgames/ASISStub/ASISConnection.java b/Engine/Plugins/Runtime/AndroidSingleInstanceService/Source/AndroidSingleInstanceService/Android/ASISStub/src/main/java/com/epicgames/ASISStub/ASISConnection.java
deleted file mode 100644
diff --git a/Engine/Plugins/Runtime/AndroidSingleInstanceService/Source/AndroidSingleInstanceService/Android/ASISStub/src/main/java/com/epicgames/ASISStub/UseServiceActivity.java b/Engine/Plugins/Runtime/AndroidSingleInstanceService/Source/AndroidSingleInstanceService/Android/ASISStub/src/main/java/com/epicgames/ASISStub/UseServiceActivity.java
deleted file mode 100644
diff --git a/Engine/Plugins/Runtime/AndroidSingleInstanceService/Source/AndroidSingleInstanceService/Android/ASISStub/src/main/java/com/epicgames/makeaar/UnrealMessageType.java b/Engine/Plugins/Runtime/AndroidSingleInstanceService/Source/AndroidSingleInstanceService/Android/ASISStub/src/main/java/com/epicgames/makeaar/UnrealMessageType.java
deleted file mode 100644
diff --git a/Engine/Plugins/Runtime/AndroidSingleInstanceService/Source/AndroidSingleInstanceService/Android/ASISStub/src/main/res/layout/interface_use_service_activity.xml b/Engine/Plugins/Runtime/AndroidSingleInstanceService/Source/AndroidSingleInstanceService/Android/ASISStub/src/main/res/layout/interface_use_service_activity.xml
deleted file mode 100644
diff --git a/Engine/Plugins/Runtime/AndroidSingleInstanceService/Source/AndroidSingleInstanceService/Android/ASISStub/src/main/res/values/attrs.xml b/Engine/Plugins/Runtime/AndroidSingleInstanceService/Source/AndroidSingleInstanceService/Android/ASISStub/src/main/res/values/attrs.xml
deleted file mode 100644
diff --git a/Engine/Plugins/Runtime/AndroidSingleInstanceService/Source/AndroidSingleInstanceService/Android/ASISStub/src/main/res/values/colors.xml b/Engine/Plugins/Runtime/AndroidSingleInstanceService/Source/AndroidSingleInstanceService/Android/ASISStub/src/main/res/values/colors.xml
deleted file mode 100644
diff --git a/Engine/Plugins/Runtime/AndroidSingleInstanceService/Source/AndroidSingleInstanceService/Android/ASISStub/src/main/res/values/dimens.xml b/Engine/Plugins/Runtime/AndroidSingleInstanceService/Source/AndroidSingleInstanceService/Android/ASISStub/src/main/res/values/dimens.xml
deleted file mode 100644
diff --git a/Engine/Plugins/Runtime/AndroidSingleInstanceService/Source/AndroidSingleInstanceService/Android/ASISStub/src/main/res/values/strings.xml b/Engine/Plugins/Runtime/AndroidSingleInstanceService/Source/AndroidSingleInstanceService/Android/ASISStub/src/main/res/values/strings.xml
deleted file mode 100644
diff --git a/Engine/Plugins/Runtime/AndroidSingleInstanceService/Source/AndroidSingleInstanceService/Android/ASISStub/src/main/res/values/styles.xml b/Engine/Plugins/Runtime/AndroidSingleInstanceService/Source/AndroidSingleInstanceService/Android/ASISStub/src/main/res/values/styles.xml
deleted file mode 100644
diff --git a/Engine/Plugins/Runtime/AndroidSingleInstanceService/Source/AndroidSingleInstanceService/Android/ASISStub/src/main/res/values/themes.xml b/Engine/Plugins/Runtime/AndroidSingleInstanceService/Source/AndroidSingleInstanceService/Android/ASISStub/src/main/res/values/themes.xml
deleted file mode 100644
diff --git a/Engine/Plugins/Runtime/AndroidSingleInstanceService/Source/AndroidSingleInstanceService/Android/gradlew.bat b/Engine/Plugins/Runtime/AndroidSingleInstanceService/Source/AndroidSingleInstanceService/Android/gradlew.bat
deleted file mode 100755
diff --git a/Engine/Plugins/Runtime/AndroidSingleInstanceService/Source/AndroidSingleInstanceService/Android/setup.bat b/Engine/Plugins/Runtime/AndroidSingleInstanceService/Source/AndroidSingleInstanceService/Android/setup.bat
deleted file mode 100755
diff --git a/Engine/Plugins/Runtime/AndroidSingleInstanceService/Source/AndroidSingleInstanceService/Android/setup.sh b/Engine/Plugins/Runtime/AndroidSingleInstanceService/Source/AndroidSingleInstanceService/Android/setup.sh
deleted file mode 100644
diff --git a/Engine/Plugins/Runtime/AndroidSingleInstanceService/Source/AndroidSingleInstanceService/Android/src/main/java/com/epicgames/makeaar/EngineFactory.java b/Engine/Plugins/Runtime/AndroidSingleInstanceService/Source/AndroidSingleInstanceService/Android/src/main/java/com/epicgames/makeaar/EngineFactory.java
--- a/Engine/Plugins/Runtime/AndroidSingleInstanceService/Source/AndroidSingleInstanceService/Android/src/main/java/com/epicgames/makeaar/EngineFactory.java
+++ b/Engine/Plugins/Runtime/AndroidSingleInstanceService/Source/AndroidSingleInstanceService/Android/src/main/java/com/epicgames/makeaar/EngineFactory.java
@@ -49,11 +49,21 @@ public class EngineFactory extends HandlerThread {
 		return Singleton.INSTANCE.getLooper();
 	}
 
+	public void onDestroy()
+	{
+		synchronized(this)
+		{
+			Log.w("UE", "LIFECYCLE: EngineFactory::onDestroy() - stop instance thread, engine = " + Singleton.INSTANCE.engine);
+
+			Singleton.INSTANCE.quit();
+		}
+	}
+
 	public static void Reset()
 	{
 		if (Singleton.INSTANCE.engine != null)
 		{
-			Log.d("UE", "EngineFactory::Reset() - releasing current engine resources, engine = " + Singleton.INSTANCE.engine);
+			Log.d("UE", "LIFECYCLE: EngineFactory::Reset() - releasing current engine resources, engine = " + Singleton.INSTANCE.engine);
 			//engine.onDestroy(543210);
 			Singleton.INSTANCE.engine.onDestroy(543210, "EngineFactory::Reset()");
 		}
@@ -66,13 +76,13 @@ public class EngineFactory extends HandlerThread {
 		
 		if (Singleton.INSTANCE.engine != null)
 		{
-            Log.d("UE", "EngineFactory::GetInstance() - USING existing Engine instance, proc = " + Application.getProcessName() + ", engine = " + Singleton.INSTANCE.engine);
+            Log.d("UE", "LIFECYCLE: EngineFactory::GetInstance() - USING existing Engine instance, proc = " + Application.getProcessName() + ", engine = " + Singleton.INSTANCE.engine);
 			return Singleton.INSTANCE.engine;
 		}
 //		engineActivity = activity;
 		Singleton.INSTANCE.engine = new Engine(activity, OBBFilename, projectModule, enablePropagateAlpha);
 		
-        Log.d("UE", "EngineFactory::GetInstance(with OBBFilename) - CREATED new Engine instance, proc = " + Application.getProcessName() + ", engine = " + Singleton.INSTANCE.engine);
+        Log.d("UE", "LIFECYCLE: EngineFactory::GetInstance(with OBBFilename) - CREATED new Engine instance, proc = " + Application.getProcessName() + ", engine = " + Singleton.INSTANCE.engine);
         return Singleton.INSTANCE.engine;
 	}
 
diff --git a/Engine/Plugins/Runtime/AndroidSingleInstanceService/Source/AndroidSingleInstanceService/Android/src/main/java/com/epicgames/makeaar/UnrealMessageType.java b/Engine/Plugins/Runtime/AndroidSingleInstanceService/Source/AndroidSingleInstanceService/Android/src/main/java/com/epicgames/makeaar/UnrealMessageType.java
--- a/Engine/Plugins/Runtime/AndroidSingleInstanceService/Source/AndroidSingleInstanceService/Android/src/main/java/com/epicgames/makeaar/UnrealMessageType.java
+++ b/Engine/Plugins/Runtime/AndroidSingleInstanceService/Source/AndroidSingleInstanceService/Android/src/main/java/com/epicgames/makeaar/UnrealMessageType.java
@@ -9,5 +9,32 @@ public enum UnrealMessageType {
 	ResumeService,
 	TouchEvent,
 	SendConsoleCommand,
-	SendData
+	SendData,
+	Unknown;
+
+	public static UnrealMessageType fromInteger(int x)
+	{
+		switch (x)
+		{
+			default:
+			case 0:
+				return Undefined;
+			case 1:
+				return Hello;
+			case 2:
+				return AttachExternalSurface;
+			case 3:
+				return DetachExternalSurface;
+			case 4:
+				return StopService;
+			case 5:
+				return ResumeService;
+			case 6:
+				return TouchEvent;
+			case 7:
+				return SendConsoleCommand;
+			case 8:
+				return SendData;
+		}
+	}
 }
\ No newline at end of file
diff --git a/Engine/Plugins/Runtime/AndroidSingleInstanceService/Source/AndroidSingleInstanceService/Android/src/main/java/com/epicgames/makeaar/UnrealSharedInstanceService.java b/Engine/Plugins/Runtime/AndroidSingleInstanceService/Source/AndroidSingleInstanceService/Android/src/main/java/com/epicgames/makeaar/UnrealSharedInstanceService.java
--- a/Engine/Plugins/Runtime/AndroidSingleInstanceService/Source/AndroidSingleInstanceService/Android/src/main/java/com/epicgames/makeaar/UnrealSharedInstanceService.java
+++ b/Engine/Plugins/Runtime/AndroidSingleInstanceService/Source/AndroidSingleInstanceService/Android/src/main/java/com/epicgames/makeaar/UnrealSharedInstanceService.java
@@ -143,51 +143,18 @@ public class UnrealSharedInstanceService extends Service {
 		}
 
 		@Override
-		public void eventCallback(int event, String param1, int param2, int param3, float param4)
+		public void eventCallback(int _event, String param1, int param2, int param3, float param4)
 		{
+			UnrealMessageType event = UnrealMessageType.fromInteger(_event);
 
-			if (event != 3 && event != 4) {
-				Log.d(TAG, "eventCallback: " + event);
-			}
-
-			if (event == Engine.EVENTTYPE_FRAME_BEGIN) {
-
-
-			}
-
-			if (event == Engine.EVENTTYPE_FRAME_END) {
-
-			}
-
-			if (event == Engine.EVENTTYPE_POST_ENGINE_INIT)
-			{
-
-
-			}
-
-			if (event == Engine.EVENTTYPE_PRE_LOAD_MAP)
-			{
-
-			}
-
-			if (event == Engine.EVENTTYPE_POST_LOAD_MAP) {
-
-				//Log.i(TAG, "Configure unlocked FPS");
-				//engine.sendConsoleCommand("rhi.SyncInterval=1"); // 0 = unlocked, 1 = 60 fps, 2 = 30 fps
-				//engine.sendConsoleCommand("r.Vsync=0");
-				//engine.sendConsoleCommand("r.GTSyncType=1"); // will cause the game thread to sync with the RHI thread (rather than the rendering thread), which caps the upper limit of input latency to something more sensible
-
-
-
-			} else if (event == Engine.EVENTTYPE_ACTION) {
-
+			if (_event != 3 && _event != 4) {
 				Log.d(TAG, "eventCallback: " + event);
 			}
 
 			if (mServiceReply != null) {
 				try {
 					Message replyMsg = Message.obtain();
-					replyMsg.what = event;
+					replyMsg.what = _event;
 					if (param1 != null && !param1.isEmpty())
 					{
 //						Log.d(TAG, "writing parceable eventCallback: " + event + ", param1=" + param1);
diff --git a/Engine/Plugins/Runtime/AndroidSingleInstanceService/Source/AndroidSingleInstanceService/AndroidSingleInstanceService_UPL.xml b/Engine/Plugins/Runtime/AndroidSingleInstanceService/Source/AndroidSingleInstanceService/AndroidSingleInstanceService_UPL.xml
--- a/Engine/Plugins/Runtime/AndroidSingleInstanceService/Source/AndroidSingleInstanceService/AndroidSingleInstanceService_UPL.xml
+++ b/Engine/Plugins/Runtime/AndroidSingleInstanceService/Source/AndroidSingleInstanceService/AndroidSingleInstanceService_UPL.xml
@@ -9,7 +9,7 @@
 		<log text="AndroidSingleInstanceService AndroidPackageName = $S(AndroidPackageName)"/>
 		
 		<!-- create the ASIS project if plugin enabled -->
-		<setBool result="bCreateASISProject" value="$B(bEnabled)"/>
+		<setBool result="bCreateASISProject" value="$B(bCompileASISProject)"/>
 		
 		<!-- by default embed ASIS in APK if plugin enabled -->
 		<setBool result="bEmbedASIS" value="$B(bEnabled)"/>
@@ -22,28 +22,6 @@
 				<setBool result="bCreateASISProject" value="false"/>
 			</false>
 		</if>
-
-
-		<!-- disable embedding ASIS if shipping but "embed in shipping" not checked  -->
-		<setBoolIsEqual result="bIsShipping" arg1="$S(Configuration)" arg2="Shipping"/>
-		<if condition="bIsShipping">
-			<true>
-				<if condition="bIncludeInShipping">
-					<false>
-						<log text="Disabled embedding ASIS because shipping"/>
-						<setBool result="bEmbedASIS" value="false"/>
-
-						<!-- but we need to compile ASISProject to install data if using ASIS -->
-						<setBool result="bCompileASISProject" value="$B(bCreateASISProject)"/>
-					</false>
-				</if>
-				<if condition="bAllowExternalStartInShipping">
-					<false>
-						<setBool result="bExternalStartASIS" value="false"/>
-					</false>
-				</if>
-			</true>
-		</if>
 	
 	</init>
 
@@ -144,6 +122,7 @@
 			<true>
 				<log text="AndroidSingleInstanceService resourceCopies copy service class files..."/>
 				<deleteFiles filespec="$S(BuildDir)/src/com/epicgames/makeaar" />
+				<deleteFiles filespec="$S(BuildDir)/gradle/ASISProject" />
 				<copyDir src="$S(PluginDir)/Android/src/main/java/com/epicgames/makeaar" dst="$S(BuildDir)/src/com/epicgames/makeaar" />
 			</true>
 			<false>
@@ -151,26 +130,34 @@
 				<deleteFiles filespec="$S(BuildDir)/src/com/epicgames/makeaar" />
 			</false>
 		</if>
+
+		<if condition="bCreateASISProject">
+			<true>
+				<log text="AndroidSingleInstanceService resourceCopies bCreateASISProject, AndroidPackageName = $S(AndroidPackageName)"/>
+				<deleteFiles filespec="$S(BuildDir)/gradle/ASISProject" />
+				<deleteFiles filespec="$S(BuildDir)/gradle/ASISProject/app/build" />
+				<copyDir src="$S(PluginDir)/Android/ASISStub" dst="$S(BuildDir)/gradle/ASISProject" />
+				<copyFile src="$S(PluginDir)/Android/src/main/java/com/epicgames/makeaar/UnrealMessageType.java" dst="$S(BuildDir)/gradle/ASISProject/app/src/main/java/com/epicgames/makeaar/UnrealMessageType.java" />
+			</true>
+		</if>
 	</resourceCopies>
 
 	<gradleCopies>
 		<if condition="bCreateASISProject">
 			<true>
-				<log text="AndroidSingleInstanceService create ASISStub project..."/>
-				<log text="AndroidSingleInstanceService bCreateASISProject, AndroidPackageName = $S(AndroidPackageName)"/>
-				<copyDir src="$S(PluginDir)/../ASISStub" dst="$S(BuildDir)/gradle/ASISProject" />
-				<copyDir src="$S(BuildDir)/gradle/app/src/main/res" dst="$S(BuildDir)/gradle/ASISProject/app/src/main/res" />
+				<log text="AndroidSingleInstanceService gradleCopies create ASISStub project..."/>
+				<log text="AndroidSingleInstanceService gradleCopies bCreateASISProject, AndroidPackageName = $S(AndroidPackageName)"/>
 				<copyDir src="$S(BuildDir)/gradle/gradle" dst="$S(BuildDir)/gradle/ASISProject/gradle" />
 			</true>
 		</if>
 	</gradleCopies>
 	
 	<buildGradleAdditions>
-		<if condition="bCompileASISProject">
+		<if condition="bCreateASISProject">
 			<true>
-				<log text="AndroidSingleInstanceService buildGradleAdditions bCompileASISProject..."/>
+				<log text="AndroidSingleInstanceService buildGradleAdditions bCreateASISProject..."/>
 				<insert>
-<![CDATA[
+					<![CDATA[
 // --- Begin ASISProject ---
 tasks.whenTaskAdded { task ->
 	if (task.name == 'assembleRelease') {
@@ -195,7 +182,21 @@ tasks.whenTaskAdded { task ->
 				commandLine 'sh', '-c', 'chmod +x ./gradlew && ./gradlew app:assembleDebug'
 			}
 		}
+		
 		task.finalizedBy 'ueASISProjectAssembleDebug'
+		
+		if (project.hasProperty('OUTPUT_PATH')) {
+			tasks.create(name:'ueASISCompleteDebug', type: Copy) {
+				description "Copy final debug APK"
+				from 'build/outputs/apk/debug/app-debug.apk'
+				into OUTPUT_PATH
+				rename ("app-debug.apk", OUTPUT_FILENAME)
+			}
+		}
+		
+		task.finalizedBy 'ueASISCompleteDebug'
+			
+
 	}
 	if (task.name == 'bundleRelease') {
 		tasks.create(name:'ueASISProjectBundleRelease', type: Exec) {
@@ -220,6 +221,17 @@ tasks.whenTaskAdded { task ->
 			}
 		}
 		task.finalizedBy 'ueASISProjectBundleDebug'
+			
+		if (project.hasProperty('OUTPUT_PATH')) {
+			tasks.create(name:'ueASISCompleteBundleDebug', type: Copy) {
+				description "Copy final debug APK"
+				from 'build/outputs/apk/debug/app-debug.apk'
+				into OUTPUT_PATH
+				rename ("app-debug.apk", OUTPUT_FILENAME)
+			}
+		}
+		
+		task.finalizedBy 'ueASISCompleteBundleDebug'
 	}
 }
 // --- End ASISProject ---
@@ -239,8 +251,9 @@ tasks.whenTaskAdded { task ->
 	<settingsGradleAdditions>
 		<if condition="bCreateASISProject">
 			<true>
+				<log text="AndroidSingleInstanceService settingsGradleAdditions updating gradle properties..."/>
 				<insert>
-	<![CDATA[
+					<![CDATA[
 	// copy files into ASISStub
 	File asis_infile = new File(rootProject.projectDir, 'gradle.properties')
 	File asis_outfile = new File(rootProject.projectDir, 'ASISProject/gradle.properties')
@@ -249,81 +262,52 @@ tasks.whenTaskAdded { task ->
 		asis_infile.withReader { reader ->
 			while ((asis_line = reader.readLine()) != null) {
 				if (asis_line.startsWith('OUTPUT_FILENAME=')) {
-					writer.writeLine("OUTPUT_FILENAME=ASIS_" + asis_line.substring(16));
+					def fname = asis_line.substring(16);
+					writer.writeLine("OUTPUT_FILENAME=ASIS_" + fname);
+					writer.writeLine("APP_NAME=ASIS_" + fname.substring(0, fname.indexOf("-Android")));
 					continue;
 				}
 				else if (asis_line.startsWith('OUTPUT_BUNDLEFILENAME=')) {
 					writer.writeLine("OUTPUT_BUNDLEFILENAME=ASIS_" + asis_line.substring(22));
 					continue;
 				}
-				writer.writeLine("${asis_line}")
-			}
-		}
-	}
-
-	asis_infile = new File(rootProject.projectDir, 'app/buildAdditions.gradle')
-	asis_outfile = new File(rootProject.projectDir, 'ASISProject/app/buildAdditions.gradle')
-	asis_line = "";
-	def bRemovingASISBlock = false;
-	asis_outfile.withWriter { writer ->
-		asis_infile.withReader { reader ->
-			while ((asis_line = reader.readLine()) != null) {
-				if (bRemovingASISBlock)
-				{
-					if (asis_line.startsWith('// --- End ASISProject ---'))
-					{
-						bRemovingASISBlock = false;
-					}
-					continue;
-				}
-				if (asis_line.startsWith('// --- Begin ASISProject ---'))
-				{
-					bRemovingASISBlock = true;
+				else if (asis_line.startsWith('PACKAGE_NAME=')) {
+					writer.writeLine("PACKAGE_NAME=com.epicgames.ASISStub");
+					writer.writeLine("ASIS_PACKAGE_NAME=" + asis_line.substring(13));
 					continue;
 				}
 				writer.writeLine("${asis_line}")
 			}
 		}
 	}
-
-	asis_infile = new File(rootProject.projectDir, 'app/src/main/AndroidManifest.xml')
-	asis_outfile = new File(rootProject.projectDir, 'ASISProject/app/src/main/AndroidManifest.xml')
-	def asis_hasForeServiceDataPerm = false;
-	def asis_hasForeServicePerm = false;
-	def asis_hasNotificationsPerm = false;
-	def asis_hasService = false
-	def asis_hasActivity = false
-	asis_outfile.withWriter { writer ->
-		asis_infile.withReader { reader ->
-			while ((asis_line = reader.readLine()) != null) {
-				if (asis_line.contains('android.permission.FOREGROUND_SERVICE_DATA_SYNC')) {
-					asis_hasForeServiceDataPerm = true
-				} else if (asis_line.contains('android.permission.FOREGROUND_SERVICE')) {
-					asis_hasForeServicePerm = true
-				} else if (asis_line.contains('android.permission.POST_NOTIFICATIONS')) {
-					asis_hasNotificationsPerm = true
-				} else if (asis_line.contains('com.epicgames.makeaar.UnrealSharedInstanceService')) {
-					asis_hasService = true
-				} else if (asis_line.contains('com.epicgames.makeaar.ASISExampleGameActivity')) {
-					asis_hasActivity = true
-				} else if (asis_line.contains('</application>')) {
-
-					if (asis_hasActivity == false) {
-						writer.writeLine('    <activity android:name="com.epicgames.makeaar.GameActivity" android:exported="true" android:screenOrientation="sensor">')
-						writer.writeLine('    	<intent-filter>')
-						writer.writeLine('    		<action android:name="android.intent.action.MAIN" />')
-						writer.writeLine('    		<category android:name="android.intent.category.LAUNCHER" />')
-						writer.writeLine('    	</intent-filter>')
-						writer.writeLine('    </activity>')
-					}
-				}
-				writer.writeLine("${asis_line}")
-			}
-		}
-	}
-
-	def asis_files = ["app/aar-imports.gradle", "app/abi.gradle", "app/build.gradle",
-		"baseBuildAdditions.gradle", "build.gradle", "buildscriptAdditions.gradle", "local.properties",
+	
+//	asis_infile = new File(rootProject.projectDir, 'app/buildAdditions.gradle')
+//	asis_outfile = new File(rootProject.projectDir, 'ASISProject/app/buildAdditions.gradle')
+//	asis_line = "";
+//	def bRemovingASISBlock = false;
+//	asis_outfile.withWriter { writer ->
+//		asis_infile.withReader { reader ->
+//			while ((asis_line = reader.readLine()) != null) {
+//				if (bRemovingASISBlock)
+//				{
+//					if (asis_line.startsWith('// --- End ASISProject ---'))
+//					{
+//						bRemovingASISBlock = false;
+//					}
+//					continue;
+//				}
+//				if (asis_line.startsWith('// --- Begin ASISProject ---'))
+//				{
+//					bRemovingASISBlock = true;
+//					continue;
+//				}
+//				writer.writeLine("${asis_line}")
+//			}
+//		}
+//	}
+
+	def asis_files = ["app/proguard-android.text",
+		"baseBuildAdditions.gradle", "buildscriptAdditions.gradle", "local.properties",
 		"gradlew", "gradlew.bat", "rungradle.bat"] as String[]
 	asis_files.eachWithIndex { v, i ->
 		asis_infile = new File(rootProject.projectDir, v)
@@ -338,7 +322,7 @@ tasks.whenTaskAdded { task ->
 			}
 		}
 	}
-
+	
 ]]>
 				</insert>
 			</true>
@@ -367,37 +351,24 @@ import java.util.Set;
 		int ConnectionType = 0;
         for (Thread thread : threads) {
             String name = thread.getName();
-            if (name.equals("UnrealASIS-USB"))
+            if (name.contains("unrealSharedInstanceService"))
             {
                 if (thread.isAlive()) {
 					ConnectionType |= 1;
                 }
             }
-            else if (name.equals("UnrealASIS-WiFi"))
-            {
-                if (thread.isAlive()) {
-					ConnectionType |= 2;
-                }
-            }
         }
 		
 		return ConnectionType;
 	}
 
-	public boolean AndroidThunkJava_AndroidSingleInstanceServer_Stop(boolean bUSB, boolean bNetwork)
+	public boolean AndroidThunkJava_AndroidSingleInstanceServer_Stop()
 	{
         Set<Thread> threads = Thread.getAllStackTraces().keySet();
         boolean bFoundThread = false;
         for (Thread thread : threads) {
             String name = thread.getName();
-            if (bUSB && name.equals("UnrealASIS-USB"))
-            {
-                if (thread.isAlive()) {
-					thread.interrupt();
-                    bFoundThread = true;
-                }
-            }
-            else if (bNetwork && name.equals("UnrealASIS-WiFi"))
+            if (name.contains("unrealSharedInstanceService"))
             {
                 if (thread.isAlive()) {
 					thread.interrupt();
@@ -408,114 +379,6 @@ import java.util.Set;
 		
 		return bFoundThread;
 	}
-
-	public boolean AndroidThunkJava_AndroidSingleInstanceServer_Start(boolean bUSB, boolean bNetwork, int port)
-	{
-		boolean bUSBRunning = false;
-		boolean bWiFiRunning = false;
-		
-		if (port == 0)
-		{
-			port = 57099;
-		}
-		
-        Set<Thread> threads = Thread.getAllStackTraces().keySet();
-        boolean bFoundThread = false;
-        for (Thread thread : threads) {
-            String name = thread.getName();
-            if (name.equals("UnrealASIS-USB"))
-            {
-                if (thread.isAlive()) {
-                    bUSBRunning = true;
-                }
-            }
-            else if (name.equals("UnrealASIS-WiFi"))
-            {
-                if (thread.isAlive()) {
-                    bWiFiRunning = true;
-                    break;
-                }
-            }
-        }
-
-		if (bUSBRunning && bWiFiRunning)
-		{
-			return true;
-		}
-
-		String packageName = getPackageName();
-		PackageManager pm = getPackageManager();
-		
-		int versionCode = 0;
-		try
-		{
-			versionCode = pm.getPackageInfo(packageName, 0).versionCode;
-		}
-		catch (Exception e)
-		{
-			// if the above failed, then, we can't use obbs
-		}
-
-		String EngineVersion = "5.0.0";
-		String ProjectName = packageName;
-		ProjectName = ProjectName.substring(ProjectName.lastIndexOf('.') + 1);
-		boolean IsShipping = false;
-		boolean UseExternalFilesDir = false;
-		boolean PublicLogFiles = false;
-		boolean AllowNetworkConnection = false;
-		try {
-			ApplicationInfo ai = pm.getApplicationInfo(packageName, PackageManager.GET_META_DATA);
-			Bundle bundle = ai.metaData;
-
-			if (bundle.containsKey("com.epicgames.unreal.GameActivity.EngineVersion"))
-			{
-				EngineVersion = bundle.getString("com.epicgames.unreal.GameActivity.EngineVersion");
-			}
-			if (bundle.containsKey("com.epicgames.unreal.GameActivity.ProjectName"))
-			{
-				ProjectName = bundle.getString("com.epicgames.unreal.GameActivity.ProjectName");
-			}
-			if (bundle.containsKey("com.epicgames.unreal.GameActivity.BuildConfiguration"))
-			{
-				String Configuration = bundle.getString("com.epicgames.unreal.GameActivity.BuildConfiguration");
-				IsShipping = Configuration.contains("Shipping");
-			}
-			if (bundle.containsKey("com.epicgames.unreal.GameActivity.bUseExternalFilesDir"))
-			{
-				UseExternalFilesDir = bundle.getBoolean("com.epicgames.unreal.GameActivity.bUseExternalFilesDir");
-			}
-			if (bundle.containsKey("com.epicgames.unreal.GameActivity.bPublicLogFiles"))
-			{
-				PublicLogFiles = bundle.getBoolean("com.epicgames.unreal.GameActivity.bPublicLogFiles");
-			}
-			if (bundle.containsKey("com.epicgames.unreal.RemoteFileManager.bAllowNetworkConnection"))
-			{
-				AllowNetworkConnection = bundle.getBoolean("com.epicgames.unreal.RemoteFileManager.bAllowNetworkConnection");
-			}
-		}
-		catch (PackageManager.NameNotFoundException | NullPointerException e)
-		{
-			Log.debug("Error when accessing application metadata");
-		}
-
-		String internal = getFilesDir().getAbsolutePath();
-		String external = getExternalFilesDir(null).getAbsolutePath();
-		String storage = android.os.Environment.getExternalStorageDirectory().getAbsolutePath();
-		String obbdir = getObbDir().getAbsolutePath();
-
-		if (bUSB && !bUSBRunning)
-		{
-			//RemoteFileManager fileManager = new RemoteFileManager(true, port, internal, external, storage, obbdir, packageName, versionCode, ProjectName, EngineVersion, IsShipping, PublicLogFiles);
-			//new Thread(fileManager, "UnrealASIS-USB").start();
-		}
-		if (AllowNetworkConnection && bNetwork && !bWiFiRunning)
-		{
-			//RemoteFileManager fileManager = new RemoteFileManager(false, port, internal, external, storage, obbdir, packageName, versionCode, ProjectName, EngineVersion, IsShipping, PublicLogFiles);
-			//new Thread(fileManager, "UnrealASIS-WiFi").start();
-		}
-		
-		return true;
-	}
 ]]>
 				</insert>
 			</true>
@@ -525,11 +388,6 @@ import java.util.Set;
 	return false;
 	}
 
-	public boolean AndroidThunkJava_AndroidSingleInstanceServer_Start()
-	{
-	return false;
-	}
-
 	public boolean AndroidThunkJava_AndroidSingleInstanceServer_Stop()
 	{
 	return false;
